new gameplay graph plugin to support gameplay features that require graphs

diff --git a/Engine/Plugins/Experimental/GameplayGraph/GameplayGraph.uplugin b/Engine/Plugins/Experimental/GameplayGraph/GameplayGraph.uplugin
new file mode 100644
--- /dev/null
+++ b/Engine/Plugins/Experimental/GameplayGraph/GameplayGraph.uplugin
@@ -0,0 +1,28 @@
+{
+	"FileVersion" : 3,
+	"Version" : 1,
+	"VersionName" : "0.1",
+	"FriendlyName" : "Gameplay Graph",
+	"Description" : "Graphs are fun.",
+	"Category": "Gameplay",
+	"CreatedBy" : "Epic Games, Inc.",
+	"CreatedByURL" : "http://epicgames.com",
+	"DocsURL" : "",
+	"MarketplaceURL" : "",
+	"SupportURL" : "",
+	"ExplicitlyLoaded": true,
+	"EnabledByDefault" : false,
+	"CanContainContent" : false,
+	"IsBetaVersion" : false,
+	"IsExperimentalVersion": true,
+	"Installed" : false,
+	"Modules": [
+		{
+			"Name": "GameplayGraph",
+			"Type": "Runtime",
+			"LoadingPhase": "Default"
+		}
+	],
+	"Plugins": [
+	]
+}
diff --git a/Engine/Plugins/Experimental/GameplayGraph/Source/GameplayGraph/GameplayGraph.Build.cs b/Engine/Plugins/Experimental/GameplayGraph/Source/GameplayGraph/GameplayGraph.Build.cs
new file mode 100644
--- /dev/null
+++ b/Engine/Plugins/Experimental/GameplayGraph/Source/GameplayGraph/GameplayGraph.Build.cs
@@ -0,0 +1,20 @@
+// Copyright Epic Games, Inc. All Rights Reserved.
+
+namespace UnrealBuildTool.Rules
+{
+	public class GameplayGraph : ModuleRules
+	{
+        public GameplayGraph(ReadOnlyTargetRules Target) : base(Target)
+		{
+			PCHUsage = ModuleRules.PCHUsageMode.UseExplicitOrSharedPCHs;
+
+			PublicDependencyModuleNames.AddRange(
+				new string[]
+				{
+					"Core",
+					"CoreUObject"
+				}
+			);
+		}
+	}
+}
diff --git a/Engine/Plugins/Experimental/GameplayGraph/Source/GameplayGraph/GameplayGraphModule.cpp b/Engine/Plugins/Experimental/GameplayGraph/Source/GameplayGraph/GameplayGraphModule.cpp
new file mode 100644
--- /dev/null
+++ b/Engine/Plugins/Experimental/GameplayGraph/Source/GameplayGraph/GameplayGraphModule.cpp
@@ -0,0 +1,5 @@
+// Copyright Epic Games, Inc. All Rights Reserved.
+
+#include "Modules/ModuleManager.h"
+
+IMPLEMENT_MODULE(FDefaultModuleImpl, GameplayGraph)
\ No newline at end of file
diff --git a/Engine/Plugins/Experimental/GameplayGraph/Source/GameplayGraph/Private/Graph/Algorithms/Connectivity/ConnectedComponents.cpp b/Engine/Plugins/Experimental/GameplayGraph/Source/GameplayGraph/Private/Graph/Algorithms/Connectivity/ConnectedComponents.cpp
new file mode 100644
--- /dev/null
+++ b/Engine/Plugins/Experimental/GameplayGraph/Source/GameplayGraph/Private/Graph/Algorithms/Connectivity/ConnectedComponents.cpp
@@ -0,0 +1,41 @@
+// Copyright Epic Games, Inc. All Rights Reserved.
+
+#include "Graph/Algorithms/Connectivity/ConnectedComponents.h"
+#include "Graph/Algorithms/Search/Search.h"
+
+namespace Graph::Algorithms
+{
+
+	TArray<TSet<FGraphVertexHandle>> FindConnectedComponents(const TSet<FGraphVertexHandle>& StartingSet)
+	{
+		TArray<TSet<FGraphVertexHandle>> Output;
+		// Do a BFS from each node already assigned to an island.
+		// Every node that we hit that's also in the StartingSet gets
+		// added to the connected component we're in the middle of
+		// building. When the traversal stops, and there's still
+		// un-touched nodes in the starting set, we create a new
+		// connected component and start again.
+		TSet<FGraphVertexHandle> WorkQueue = StartingSet;
+		while (!WorkQueue.IsEmpty())
+		{
+			TSet<FGraphVertexHandle> CurrentComponent;
+			BFS(
+				*WorkQueue.CreateConstIterator(),
+				[&CurrentComponent, &WorkQueue, &StartingSet](const FGraphVertexHandle& Handle)
+				{
+					if (StartingSet.Contains(Handle))
+					{
+						CurrentComponent.Add(Handle);
+					}
+					WorkQueue.Remove(Handle);
+					return false;
+				}
+			);
+
+			Output.Add(CurrentComponent);
+		}
+
+		return Output;
+	}
+
+}
\ No newline at end of file
diff --git a/Engine/Plugins/Experimental/GameplayGraph/Source/GameplayGraph/Private/Graph/Algorithms/Search/Search.cpp b/Engine/Plugins/Experimental/GameplayGraph/Source/GameplayGraph/Private/Graph/Algorithms/Search/Search.cpp
new file mode 100644
--- /dev/null
+++ b/Engine/Plugins/Experimental/GameplayGraph/Source/GameplayGraph/Private/Graph/Algorithms/Search/Search.cpp
@@ -0,0 +1,114 @@
+// Copyright Epic Games, Inc. All Rights Reserved.
+#include "Graph/Algorithms/Search/Search.h"
+
+#include "Containers/Array.h"
+#include "Containers/Queue.h"
+#include "Containers/Set.h"
+#include "Graph/GraphVertex.h"
+
+namespace Graph::Algorithms
+{
+	namespace
+	{
+		using BFSDataStructure = TQueue<FGraphVertexHandle>;
+		using DFSDataStructure = TArray<FGraphVertexHandle>;
+
+		template<template<typename> typename TDataStructure>
+		FGraphVertexHandle GetNextAndAdvance(TDataStructure<FGraphVertexHandle>& Queue)
+		{
+			if constexpr (std::is_same_v<TDataStructure<FGraphVertexHandle>, BFSDataStructure>)
+			{
+				FGraphVertexHandle Ret;
+				if (ensure(Queue.Dequeue(Ret)))
+				{
+					return Ret;
+				}
+				return {};
+			}
+			else if constexpr (std::is_same_v<TDataStructure<FGraphVertexHandle>, DFSDataStructure>)
+			{
+				if (ensure(!Queue.IsEmpty()))
+				{
+					FGraphVertexHandle Next = Queue.Top();
+					Queue.Pop();
+					return Next;
+				}
+
+				return {};
+			}
+			else
+			{
+				static_assert("Invalid data structure for graph search [GetNextAndAdvance].");
+			}
+		}
+
+		template<template<typename> typename TDataStructure>
+		void AddToWorkQueue(TDataStructure<FGraphVertexHandle>& Queue, const FGraphVertexHandle& Handle)
+		{
+			if constexpr (std::is_same_v<TDataStructure<FGraphVertexHandle>, BFSDataStructure>)
+			{
+				Queue.Enqueue(Handle);
+			}
+			else if constexpr (std::is_same_v<TDataStructure<FGraphVertexHandle>, DFSDataStructure>)
+			{
+				Queue.Add(Handle);
+			}
+			else
+			{
+				static_assert("Invalid data structure for graph search [AddToWorkQueue].");
+			}
+		}
+
+		template<template<typename> typename TDataStructure>
+		FGraphVertexHandle GenericSearch(const FGraphVertexHandle& Start, FSearchCallback Callback)
+		{
+			TDataStructure<FGraphVertexHandle> WorkQueue;
+			AddToWorkQueue<TDataStructure>(WorkQueue, Start);
+
+			TSet<FGraphVertexHandle> Seen;
+			Seen.Add(Start);
+
+			while (!WorkQueue.IsEmpty())
+			{
+				FGraphVertexHandle Next = GetNextAndAdvance<TDataStructure>(WorkQueue);
+				if (!Next.IsComplete())
+				{
+					continue;
+				}
+
+				if (Callback(Next))
+				{
+					return Next;
+				}
+
+				// Get neighbors and add to the queue.
+				if (TObjectPtr<UGraphVertex> Node = Next.GetVertex())
+				{
+					Node->ForEachAdjacentVertex(
+						[&WorkQueue, &Seen](const FGraphVertexHandle& Neighbor)
+						{
+							if (Neighbor.IsComplete() && !Seen.Contains(Neighbor))
+							{
+								AddToWorkQueue<TDataStructure>(WorkQueue, Neighbor);
+								Seen.Add(Neighbor);
+							}
+						}
+					);
+				}
+			}
+
+			return {};
+		}
+
+	}
+
+	FGraphVertexHandle BFS(const FGraphVertexHandle& Start, FSearchCallback Callback)
+	{
+		return GenericSearch<TQueue>(Start, Callback);
+	}
+
+	FGraphVertexHandle DFS(const FGraphVertexHandle& Start, FSearchCallback Callback)
+	{
+		return GenericSearch<TArray>(Start, Callback);
+	}
+}
\ No newline at end of file
diff --git a/Engine/Plugins/Experimental/GameplayGraph/Source/GameplayGraph/Private/Graph/Graph.cpp b/Engine/Plugins/Experimental/GameplayGraph/Source/GameplayGraph/Private/Graph/Graph.cpp
new file mode 100644
--- /dev/null
+++ b/Engine/Plugins/Experimental/GameplayGraph/Source/GameplayGraph/Private/Graph/Graph.cpp
@@ -0,0 +1,511 @@
+// Copyright Epic Games, Inc. All Rights Reserved.
+
+#include "Graph/Graph.h"
+#include "Graph/GraphEdge.h"
+#include "Graph/GraphVertex.h"
+#include "Graph/Algorithms/Connectivity/ConnectedComponents.h"
+#include "GenericPlatform/GenericPlatformMath.h"
+
+#include UE_INLINE_GENERATED_CPP_BY_NAME(Graph)
+
+FSerializableGraph UGraph::GetSerializableGraph() const
+{
+	FSerializableGraph Out;
+	Out.Properties = Properties;
+
+	for (const TPair<FGraphVertexHandle, TObjectPtr<UGraphVertex>>& Kvp : Vertices)
+	{
+		if (!Kvp.Key.IsComplete() || !Kvp.Value)
+		{
+			continue;
+		}
+		Out.Vertices.Add(Kvp.Key);
+	}
+
+	for (const TPair<FGraphEdgeHandle, TObjectPtr<UGraphEdge>>& Kvp : Edges)
+	{
+		if (!Kvp.Key.IsComplete() || !Kvp.Value)
+		{
+			continue;
+		}
+
+		Out.Edges.Add(Kvp.Key, Kvp.Value->GetSerializedData());
+	}
+
+	for (const TPair<FGraphIslandHandle, TObjectPtr<UGraphIsland>>& Kvp : Islands)
+	{
+		if (!Kvp.Key.IsComplete() || !Kvp.Value)
+		{
+			continue;
+		}
+
+		Out.Islands.Add(Kvp.Key, Kvp.Value->GetSerializedData());
+	}
+
+	return Out;
+}
+
+void UGraph::LoadFromSerializedGraph(const FSerializableGraph& Input)
+{
+	InitializeFromProperties(Input.Properties);
+
+	for (const FGraphVertexHandle& Handle : Input.Vertices)
+	{
+		if (!Handle.IsValid())
+		{
+			continue;
+		}
+
+		CreateVertex(Handle.GetUniqueIndex());
+	}
+
+	for (const TPair<FGraphEdgeHandle, FSerializedEdgeData>& Kvp : Input.Edges)
+	{
+		if (!Kvp.Key.IsValid())
+		{
+			continue;
+		}
+
+		CreateEdge(
+			GetCompleteNodeHandle(Kvp.Value.Node1),
+			GetCompleteNodeHandle(Kvp.Value.Node2),
+			Kvp.Key.GetUniqueIndex(),
+			false
+		);
+	}
+
+	for (const TPair<FGraphIslandHandle, FSerializedIslandData>& Kvp : Input.Islands)
+	{
+		if (!Kvp.Key.IsValid())
+		{
+			continue;
+		}
+
+		TArray<FGraphVertexHandle> IslandVertices = Kvp.Value.Vertices;
+		for (FGraphVertexHandle& VertexHandle : IslandVertices)
+		{
+			VertexHandle = GetCompleteNodeHandle(VertexHandle);
+		}
+
+		CreateIsland(IslandVertices, Kvp.Key.GetUniqueIndex());
+	}
+}
+
+void UGraph::Empty()
+{
+	Vertices.Empty();
+	Edges.Empty();
+	Islands.Empty();
+	NextAvailableVertexUniqueIndex = 0;
+	NextAvailableEdgeUniqueIndex = 0;
+	NextAvailableIslandUniqueIndex = 0;
+}
+
+void UGraph::InitializeFromProperties(const FGraphProperties& InProperties)
+{
+	Empty();
+	Properties = InProperties;
+}
+
+FGraphVertexHandle UGraph::CreateVertex(int64 InUniqueIndex)
+{
+	TObjectPtr<UGraphVertex> Vertex = CreateTypedVertex();
+	if (!Vertex)
+	{
+		return {};
+	}
+	Vertex->OnCreate();
+
+	// No need to increment since AddNode will do that.
+	Vertex->SetUniqueIndex((InUniqueIndex == INDEX_NONE) ? NextAvailableVertexUniqueIndex : InUniqueIndex);
+	RegisterVertex(Vertex);
+	return Vertex->Handle();
+}
+
+void UGraph::RegisterVertex(TObjectPtr<UGraphVertex> Vertex)
+{
+	if (!Vertex)
+	{
+		return;
+	}
+
+	const FGraphVertexHandle Handle = Vertex->Handle();
+	Vertex->SetParentGraph(this);
+	Vertices.Add(Handle, Vertex);
+	NextAvailableVertexUniqueIndex = FMath::Max(NextAvailableVertexUniqueIndex, Handle.GetUniqueIndex() + 1);
+}
+
+FGraphEdgeHandle UGraph::CreateEdge(FGraphVertexHandle Node1, FGraphVertexHandle Node2, int64 InUniqueIndex, bool bAddToIslands)
+{
+	if (!Node1.IsComplete() || !Node2.IsComplete())
+	{
+		return {};
+	}
+
+	TObjectPtr<UGraphEdge> Edge = CreateTypedEdge();
+	if (!Edge)
+	{
+		return {};
+	}
+
+	if (Node2 < Node1)
+	{
+		std::swap(Node1, Node2);
+	}
+
+	TObjectPtr<UGraphVertex> Node1Ptr = Node1.GetVertex();
+	TObjectPtr<UGraphVertex> Node2Ptr = Node2.GetVertex();
+	if (!Node1Ptr || !Node2Ptr)
+	{
+		return {};
+	}
+
+	// For edges, we also need to make sure the edge doesn't already exist.
+	if (Node1Ptr->HasEdgeTo(Node2) || Node2Ptr->HasEdgeTo(Node1))
+	{
+		return {};
+	}
+	Edge->OnCreate();
+
+	// No need to increment since AddEdge will do that.
+	Edge->SetUniqueIndex((InUniqueIndex == INDEX_NONE) ? NextAvailableEdgeUniqueIndex : InUniqueIndex);
+	Edge->SetNodes(Node1, Node2);
+
+	RegisterEdge(Edge);
+
+	// Also need to make sure the nodes are aware of the new edge.
+	Node1Ptr->AddEdgeTo(Node2, Edge->Handle());
+	Node2Ptr->AddEdgeTo(Node1, Edge->Handle());
+
+	// If we want to keep track of islands, this is where we need to create/merge islands.
+	if (Properties.bGenerateIslands && bAddToIslands)
+	{
+		MergeOrCreateIslands(Edge);
+	}
+
+	return Edge->Handle();
+}
+
+void UGraph::RegisterEdge(TObjectPtr<UGraphEdge> Edge)
+{
+	if (!Edge)
+	{
+		return;
+	}
+
+	const FGraphEdgeHandle Handle = Edge->Handle();
+	Edges.Add(Handle, Edge);
+	Edge->SetParentGraph(this);
+
+	VertexEdges.FindOrAdd(Edge->NodeA()).Add(Handle);
+	VertexEdges.FindOrAdd(Edge->NodeB()).Add(Handle);
+
+	NextAvailableEdgeUniqueIndex = FMath::Max(NextAvailableEdgeUniqueIndex, Handle.GetUniqueIndex() + 1);
+}
+
+FGraphIslandHandle UGraph::CreateIsland(TArray<FGraphVertexHandle> InputNodes, int64 InUniqueIndex)
+{
+	if (InputNodes.IsEmpty())
+	{
+		return {};
+	}
+
+	TObjectPtr<UGraphIsland> Island = CreateTypedIsland();
+	if (!Island)
+	{
+		return {};
+	}
+	Island->OnCreate();
+
+	Island->SetUniqueIndex((InUniqueIndex == INDEX_NONE) ? NextAvailableIslandUniqueIndex : InUniqueIndex);
+	for (const FGraphVertexHandle& Node : InputNodes)
+	{
+		Island->AddVertex(Node);
+	}
+
+	RegisterIsland(Island);
+	return Island->Handle();
+}
+
+void UGraph::RegisterIsland(TObjectPtr<UGraphIsland> Island)
+{
+	if (!Island)
+	{
+		return;
+	}
+	const FGraphIslandHandle Handle = Island->Handle();
+	Island->SetParentGraph(this);
+	Islands.Add(Handle, Island);
+	NextAvailableIslandUniqueIndex = FMath::Max(NextAvailableIslandUniqueIndex, Handle.GetUniqueIndex() + 1);
+}
+
+void UGraph::RemoveIsland(const FGraphIslandHandle& IslandHandle)
+{
+	if (!IslandHandle.IsValid())
+	{
+		return;
+	}
+
+	if (TObjectPtr<UGraphIsland> Island = IslandHandle.GetIsland())
+	{
+		Island->Destroy();
+	}
+	Islands.Remove(IslandHandle);
+}
+
+void UGraph::MergeOrCreateIslands(TObjectPtr<UGraphEdge> Edge)
+{
+	if (!Edge)
+	{
+		return;
+	}
+
+	FGraphVertexHandle AHandle = Edge->NodeA();
+	FGraphVertexHandle BHandle = Edge->NodeB();
+
+	TObjectPtr<UGraphVertex> NodeA = AHandle.GetVertex();
+	TObjectPtr<UGraphVertex> NodeB = BHandle.GetVertex();
+	if (!NodeA || !NodeB)
+	{
+		return;
+	}
+
+	FGraphIslandHandle IslandHandleA = NodeA->GetParentIsland();
+	FGraphIslandHandle IslandHandleB = NodeB->GetParentIsland();
+
+	TObjectPtr<UGraphIsland> IslandA = IslandHandleA.GetIsland();
+	TObjectPtr<UGraphIsland> IslandB = IslandHandleB.GetIsland();
+
+	if (IslandHandleA.IsValid() && IslandHandleB.IsValid() && IslandHandleA != IslandHandleB)
+	{
+		// Both are in separate, valid islands - merge islands (i.e. use island A and just get rid of island B).
+		if (ensure(IslandA != nullptr))
+		{
+			IslandA->MergeWith(IslandB);
+		}
+		RemoveIsland(IslandHandleB);
+	}
+	else if (IslandHandleA.IsValid() && !IslandHandleB.IsValid())
+	{
+		// Only A is in an island, go into island A.
+		if (ensure(IslandA != nullptr))
+		{
+			IslandA->AddVertex(BHandle);
+		}
+	}
+	else if (IslandHandleB.IsValid() && !IslandHandleA.IsValid())
+	{
+		// Only B is in an island, go into island B.
+		if (ensure(IslandB != nullptr))
+		{
+			IslandB->AddVertex(AHandle);
+		}
+	}
+	else if (!IslandHandleA.IsValid() && !IslandHandleB.IsValid())
+	{
+		// Neither is in an island - need to create a new one.
+		CreateIsland({ AHandle, BHandle });
+	}
+}
+
+void UGraph::RemoveVertex(const FGraphVertexHandle& NodeHandle)
+{
+	if (!NodeHandle.IsValid())
+	{
+		return;
+	}
+
+	TObjectPtr<UGraphVertex> Node = NodeHandle.GetVertex();
+	if (!Node)
+	{
+		return;
+	}
+
+	// We must remove every edge this node is a part of.
+	for (const FGraphEdgeHandle& EdgeHandle : VertexEdges.FindOrAdd(NodeHandle))
+	{
+		// Don't immediately handle islands. We'll do it later.
+		RemoveEdge(EdgeHandle, false);
+	}
+
+	if (Properties.bGenerateIslands)
+	{
+		// If the node is a part of an island - remove it from the island and evaluate the effect that has.
+		if (const FGraphIslandHandle& IslandHandle = Node->GetParentIsland(); IslandHandle.IsValid())
+		{
+			if (TObjectPtr<UGraphIsland> Island = IslandHandle.GetIsland())
+			{
+				Island->RemoveVertex(NodeHandle);
+				RemoveOrSplitIsland(Island);
+			}
+		}
+	}
+
+	Node->HandleOnVertexRemoved();
+	Vertices.Remove(NodeHandle);
+	VertexEdges.Remove(NodeHandle);
+}
+
+void UGraph::RemoveEdge(const FGraphEdgeHandle& EdgeHandle, bool bHandleIslands)
+{
+	if (!EdgeHandle.IsValid())
+	{
+		return;
+	}
+
+	TObjectPtr<UGraphEdge> Edge = EdgeHandle.GetEdge();
+	if (!Edge)
+	{
+		return;
+	}
+
+	FGraphIslandHandle IslandHandle;
+
+	// Need to remove the edge reference from both nodes.
+	if (const FGraphVertexHandle& NodeHandleA = Edge->NodeA(); NodeHandleA.IsValid())
+	{
+		if (TObjectPtr<UGraphVertex> Node = NodeHandleA.GetVertex(); Node)
+		{
+			Node->RemoveEdge(EdgeHandle);
+			IslandHandle = Node->GetParentIsland();
+		}
+	}
+
+	if (const FGraphVertexHandle& NodeHandleB = Edge->NodeB(); NodeHandleB.IsValid())
+	{
+		if (TObjectPtr<UGraphVertex> Node = NodeHandleB.GetVertex(); Node)
+		{
+			Node->RemoveEdge(EdgeHandle);
+
+			// Technically shouldn't be necessary but in here just in case.
+			if (!IslandHandle.IsValid())
+			{
+				IslandHandle = Node->GetParentIsland();
+			}
+		}
+	}
+	
+	// Removing an edge should cause the island to check if it needs to split.
+	if (Properties.bGenerateIslands && bHandleIslands)
+	{
+		// Note that we can assume that both nodes are in the same island.
+		RemoveOrSplitIsland(IslandHandle.GetIsland());
+	}
+}
+
+void UGraph::RemoveOrSplitIsland(TObjectPtr<UGraphIsland> Island)
+{
+	if (!Island)
+	{
+		return;
+	}
+
+	if (Island->IsEmpty())
+	{
+		RemoveIsland(Island->Handle());
+	}
+	else
+	{
+		TSet<FGraphVertexHandle> IslandNodes = Island->GetVertices();
+		TArray<TSet<FGraphVertexHandle>> ConnectedComponents = Graph::Algorithms::FindConnectedComponents(Island->GetVertices());
+
+		// Sort the connected components from largest to smallest since we'll make the assumption that ConnectedComponents[0] will
+		// correspond to the already existing island. This way we do the least amount of transferal of nodes from the initial island
+		// to the new islands we're going to create.
+		ConnectedComponents.Sort(
+			[](const TSet<FGraphVertexHandle>& A, const TSet<FGraphVertexHandle>& B)
+			{
+				return A.Num() > B.Num();
+			}
+		);
+
+		TArray<TObjectPtr<UGraphIsland>> AffectedIslands = { Island };
+		const bool bIsIdentical = ConnectedComponents.Num() == 1 && ConnectedComponents[0].Num() == IslandNodes.Num();
+		if (!bIsIdentical)
+		{
+			// First remove any of the nodes that need to be removed from the existing island.
+			for (const FGraphVertexHandle& NodeHandle : IslandNodes)
+			{
+				if (!ConnectedComponents[0].Contains(NodeHandle))
+				{
+					Island->RemoveVertex(NodeHandle);
+				}
+			}
+
+			// Next, create islands for every additional connected component.
+			for (int32 Index = 1; Index < ConnectedComponents.Num(); ++Index)
+			{
+				FGraphIslandHandle NewIsland = CreateIsland(ConnectedComponents[Index].Array());
+				if (NewIsland.IsComplete())
+				{
+					AffectedIslands.Add(NewIsland.GetIsland());
+				}
+			}
+		}
+
+		// All the islands that got changed need to have OnConnectivityChanged called.
+		// This let's all the listeners of the islands make a decision on what to do with
+		// the island that had a destructive change to its connected components as well as
+		// the newly created islands.
+		for (TObjectPtr<UGraphIsland>& AffectedIsland : AffectedIslands)
+		{
+			AffectedIsland->HandleOnConnectivityChanged();
+		}
+	}
+}
+
+FGraphVertexHandle UGraph::GetCompleteNodeHandle(const FGraphVertexHandle& InHandle)
+{
+	if (!InHandle.IsValid() || InHandle.HasElement())
+	{
+		return InHandle;
+	}
+
+	TObjectPtr<UGraphVertex>* Data = Vertices.Find(InHandle);
+	if (!Data || !*Data)
+	{
+		return {};
+	}
+
+	return (*Data)->Handle();
+}
+
+TObjectPtr<UGraphVertex> UGraph::CreateTypedVertex() const
+{
+	return NewObject<UGraphVertex>(const_cast<UGraph*>(this), UGraphVertex::StaticClass());
+}
+
+TObjectPtr<UGraphEdge> UGraph::CreateTypedEdge() const
+{
+	return NewObject<UGraphEdge>(const_cast<UGraph*>(this), UGraphEdge::StaticClass());
+}
+
+TObjectPtr<UGraphIsland> UGraph::CreateTypedIsland() const
+{
+	return NewObject<UGraphIsland>(const_cast<UGraph*>(this), UGraphIsland::StaticClass());
+}
+
+void UGraph::FinalizeVertex(const FGraphVertexHandle& InHandle)
+{
+	if (!InHandle.IsComplete())
+	{
+		return;
+	}
+
+	TObjectPtr<UGraphVertex> Node = InHandle.GetVertex();
+	if (!Node)
+	{
+		return;
+	}
+
+	const FGraphIslandHandle& IslandHandle = Node->GetParentIsland();
+	if (!IslandHandle.IsComplete())
+	{
+		return;
+	}
+
+	if (TObjectPtr<UGraphIsland> Island = IslandHandle.GetIsland())
+	{
+		Island->HandleOnConnectivityChanged();
+	}
+}
\ No newline at end of file
diff --git a/Engine/Plugins/Experimental/GameplayGraph/Source/GameplayGraph/Private/Graph/GraphEdge.cpp b/Engine/Plugins/Experimental/GameplayGraph/Source/GameplayGraph/Private/Graph/GraphEdge.cpp
new file mode 100644
--- /dev/null
+++ b/Engine/Plugins/Experimental/GameplayGraph/Source/GameplayGraph/Private/Graph/GraphEdge.cpp
@@ -0,0 +1,41 @@
+// Copyright Epic Games, Inc. All Rights Reserved.
+
+#include "Graph/GraphEdge.h"
+
+#include UE_INLINE_GENERATED_CPP_BY_NAME(GraphEdge)
+
+UGraphEdge::UGraphEdge()
+	: Super(EGraphElementType::Edge)
+{
+}
+
+void UGraphEdge::SetNodes(const FGraphVertexHandle& InA, const FGraphVertexHandle& InB)
+{
+	A = InA;
+	B = InB;
+}
+
+const FGraphVertexHandle& UGraphEdge::GetOtherNode(const FGraphVertexHandle& InNode) const
+{
+	if (InNode == A)
+	{
+		return B;
+	}
+	else
+	{
+		return A;
+	}
+}
+
+bool UGraphEdge::ContainsNode(const FGraphVertexHandle& InNode) const
+{
+	return (InNode == A) || (InNode == B);
+}
+
+FSerializedEdgeData UGraphEdge::GetSerializedData() const
+{
+	FSerializedEdgeData Nodes;
+	Nodes.Node1 = A;
+	Nodes.Node2 = B;
+	return Nodes;
+}
\ No newline at end of file
diff --git a/Engine/Plugins/Experimental/GameplayGraph/Source/GameplayGraph/Private/Graph/GraphHandle.cpp b/Engine/Plugins/Experimental/GameplayGraph/Source/GameplayGraph/Private/Graph/GraphHandle.cpp
new file mode 100644
--- /dev/null
+++ b/Engine/Plugins/Experimental/GameplayGraph/Source/GameplayGraph/Private/Graph/GraphHandle.cpp
@@ -0,0 +1,99 @@
+// Copyright Epic Games, Inc. All Rights Reserved.
+
+#include "Graph/GraphHandle.h"
+#include "Graph/GraphElement.h"
+#include "Graph/GraphEdge.h"
+#include "Graph/GraphVertex.h"
+
+#include UE_INLINE_GENERATED_CPP_BY_NAME(GraphHandle)
+
+FGraphHandle::FGraphHandle()
+{
+
+}
+
+FGraphHandle::FGraphHandle(int64 InUniqueIndex, TObjectPtr<UGraphElement> InElement)
+	: UniqueIndex(InUniqueIndex)
+	, Element(InElement)
+{
+}
+
+void FGraphHandle::SetElement(TObjectPtr<UGraphElement> InElement)
+{
+	Element = InElement;
+}
+
+TObjectPtr<UGraphElement> FGraphHandle::GetElement() const
+{
+	return Element.Get();
+}
+
+bool FGraphHandle::IsValid() const
+{
+	return UniqueIndex != INDEX_NONE;
+}
+
+bool FGraphHandle::HasElement() const
+{
+	return Element.IsValid();
+}
+
+bool FGraphHandle::IsComplete() const
+{
+	return IsValid() && HasElement();
+}
+
+bool FGraphHandle::operator==(const FGraphHandle& Other) const
+{
+	return UniqueIndex == Other.UniqueIndex;
+}
+
+bool FGraphHandle::operator!=(const FGraphHandle& Other) const
+{
+	return UniqueIndex != Other.UniqueIndex;
+}
+
+bool FGraphHandle::operator<(const FGraphHandle& Other) const
+{
+	return UniqueIndex < Other.UniqueIndex;
+}
+
+void FGraphHandle::Clear()
+{
+	UniqueIndex = INDEX_NONE;
+	Element = nullptr;
+}
+
+uint32 GetTypeHash(const FGraphHandle& Handle)
+{
+	return ::GetTypeHash(Handle.UniqueIndex);
+}
+
+FGraphVertexHandle::FGraphVertexHandle()
+{
+}
+
+TObjectPtr<UGraphVertex> FGraphVertexHandle::GetVertex() const
+{
+	return Cast<UGraphVertex>(GetElement());
+}
+
+FGraphEdgeHandle::FGraphEdgeHandle()
+{
+
+}
+
+TObjectPtr<UGraphEdge> FGraphEdgeHandle::GetEdge() const
+{
+	return Cast<UGraphEdge>(GetElement());
+}
+
+FGraphIslandHandle::FGraphIslandHandle()
+{
+
+}
+
+TObjectPtr<UGraphIsland> FGraphIslandHandle::GetIsland() const
+{
+	return Cast<UGraphIsland>(GetElement());
+}
\ No newline at end of file
diff --git a/Engine/Plugins/Experimental/GameplayGraph/Source/GameplayGraph/Private/Graph/GraphIsland.cpp b/Engine/Plugins/Experimental/GameplayGraph/Source/GameplayGraph/Private/Graph/GraphIsland.cpp
new file mode 100644
--- /dev/null
+++ b/Engine/Plugins/Experimental/GameplayGraph/Source/GameplayGraph/Private/Graph/GraphIsland.cpp
@@ -0,0 +1,112 @@
+// Copyright Epic Games, Inc. All Rights Reserved.
+
+#include "Graph/GraphIsland.h"
+#include "Graph/GraphVertex.h"
+#include "Graph/Graph.h"
+
+#include UE_INLINE_GENERATED_CPP_BY_NAME(GraphIsland)
+
+UGraphIsland::UGraphIsland()
+	: Super(EGraphElementType::Island)
+{
+}
+
+void UGraphIsland::Destroy()
+{
+	TSet<FGraphVertexHandle> VertexCopy = Vertices;
+	Vertices.Empty();
+
+	// Need to work off a copy of the vertex handles since calling
+	// UGraph::RemoveVertex will attempt to remove the node from its
+	// parent island as well and that'll run into an error of modifying
+	// the Vertices TSet during iteration.
+	for (const FGraphVertexHandle& Node : VertexCopy)
+	{
+		if (TObjectPtr<UGraph> Graph = GetGraph())
+		{
+			Graph->RemoveVertex(Node);
+		}
+	}
+
+	Vertices.Empty();
+	HandleOnDestroyed();
+}
+
+void UGraphIsland::MergeWith(TObjectPtr<UGraphIsland> OtherIsland)
+{
+	if (!OtherIsland)
+	{
+		return;
+	}
+
+	for (const FGraphVertexHandle& Node : OtherIsland->Vertices)
+	{
+		AddVertex(Node);
+	}
+
+	OtherIsland->Vertices.Empty();
+}
+
+void UGraphIsland::AddVertex(const FGraphVertexHandle& Node)
+{
+	if (!Node.IsValid())
+	{
+		return;
+	}
+
+	TObjectPtr<UGraphVertex> NodePtr = Node.GetVertex();
+	if (!NodePtr)
+	{
+		return;
+	}
+	NodePtr->SetParentIsland(Handle());
+	Vertices.Add(Node);
+	HandleOnVertexAdded(Node);
+}
+
+void UGraphIsland::RemoveVertex(const FGraphVertexHandle& Node)
+{
+	if (!Node.IsValid())
+	{
+		return;
+	}
+
+	TObjectPtr<UGraphVertex> NodePtr = Node.GetVertex();
+	if (!NodePtr)
+	{
+		return;
+	}
+	NodePtr->SetParentIsland({});
+	Vertices.Remove(Node);
+	HandleOnVertexRemoved(Node);
+}
+
+FSerializedIslandData UGraphIsland::GetSerializedData() const
+{
+	FSerializedIslandData Out;
+	for (const FGraphVertexHandle& Handle : Vertices)
+	{
+		Out.Vertices.Add(Handle);
+	}
+	return Out;
+}
+
+void UGraphIsland::HandleOnVertexAdded(const FGraphVertexHandle& Handle)
+{
+	OnVertexAdded.Broadcast(Handle);
+}
+
+void UGraphIsland::HandleOnVertexRemoved(const FGraphVertexHandle& Handle)
+{
+	OnVertexRemoved.Broadcast(Handle);
+}
+
+void UGraphIsland::HandleOnDestroyed()
+{
+	OnDestroyed.Broadcast();
+}
+
+void UGraphIsland::HandleOnConnectivityChanged()
+{
+	OnConnectivityChanged.Broadcast();
+}
\ No newline at end of file
diff --git a/Engine/Plugins/Experimental/GameplayGraph/Source/GameplayGraph/Private/Graph/GraphVertex.cpp b/Engine/Plugins/Experimental/GameplayGraph/Source/GameplayGraph/Private/Graph/GraphVertex.cpp
new file mode 100644
--- /dev/null
+++ b/Engine/Plugins/Experimental/GameplayGraph/Source/GameplayGraph/Private/Graph/GraphVertex.cpp
@@ -0,0 +1,42 @@
+// Copyright Epic Games, Inc. All Rights Reserved.
+
+#include "Graph/GraphVertex.h"
+#include "Graph/GraphEdge.h"
+
+#include UE_INLINE_GENERATED_CPP_BY_NAME(GraphVertex)
+
+UGraphVertex::UGraphVertex()
+	: Super(EGraphElementType::Node)
+{
+}
+
+bool UGraphVertex::HasEdgeTo(const FGraphVertexHandle& Other) const
+{
+	return Edges.Contains(Other);
+}
+
+void UGraphVertex::AddEdgeTo(const FGraphVertexHandle& Node, const FGraphEdgeHandle& Edge)
+{
+	Edges.Add(Node, Edge);
+}
+
+void UGraphVertex::RemoveEdge(const FGraphEdgeHandle& EdgeHandle)
+{
+	if (!EdgeHandle.IsValid())
+	{
+		return;
+	}
+
+	TObjectPtr<UGraphEdge> Edge = EdgeHandle.GetEdge();
+	if (!Edge || !Edge->ContainsNode(Handle()))
+	{
+		return;
+	}
+
+	Edges.Remove(Edge->GetOtherNode(Handle()));
+}
+
+void UGraphVertex::HandleOnVertexRemoved()
+{
+	OnVertexRemoved.Broadcast();
+}
\ No newline at end of file
diff --git a/Engine/Plugins/Experimental/GameplayGraph/Source/GameplayGraph/Public/Graph/Algorithms/Connectivity/ConnectedComponents.h b/Engine/Plugins/Experimental/GameplayGraph/Source/GameplayGraph/Public/Graph/Algorithms/Connectivity/ConnectedComponents.h
new file mode 100644
--- /dev/null
+++ b/Engine/Plugins/Experimental/GameplayGraph/Source/GameplayGraph/Public/Graph/Algorithms/Connectivity/ConnectedComponents.h
@@ -0,0 +1,17 @@
+// Copyright Epic Games, Inc. All Rights Reserved.
+#pragma once
+
+#include "Containers/Array.h"
+#include "Containers/Set.h"
+#include "Graph/GraphHandle.h"
+
+namespace Graph::Algorithms
+{
+	/**
+	 * Given a set of nodes, returns an array of connected components.
+	 * Note that any nodes that may be otherwise connected to the input nodes will be ignored
+	 * for the output if they are not in the input set.
+	 */
+	TArray<TSet<FGraphVertexHandle>> FindConnectedComponents(const TSet<FGraphVertexHandle>& StartingSet);
+
+}
\ No newline at end of file
diff --git a/Engine/Plugins/Experimental/GameplayGraph/Source/GameplayGraph/Public/Graph/Algorithms/Search/Search.h b/Engine/Plugins/Experimental/GameplayGraph/Source/GameplayGraph/Public/Graph/Algorithms/Search/Search.h
new file mode 100644
--- /dev/null
+++ b/Engine/Plugins/Experimental/GameplayGraph/Source/GameplayGraph/Public/Graph/Algorithms/Search/Search.h
@@ -0,0 +1,13 @@
+// Copyright Epic Games, Inc. All Rights Reserved.
+#pragma once
+
+#include "Graph/GraphHandle.h"
+
+namespace Graph::Algorithms
+{
+	/** Returns whether we've found what we've wanted to find (i.e. returning true will end the search). */
+	using FSearchCallback = TFunction<bool(const FGraphVertexHandle&)>;
+
+	FGraphVertexHandle BFS(const FGraphVertexHandle& Start, FSearchCallback Callback);
+	FGraphVertexHandle DFS(const FGraphVertexHandle& Start, FSearchCallback Callback);
+}
\ No newline at end of file
diff --git a/Engine/Plugins/Experimental/GameplayGraph/Source/GameplayGraph/Public/Graph/Graph.h b/Engine/Plugins/Experimental/GameplayGraph/Source/GameplayGraph/Public/Graph/Graph.h
new file mode 100644
--- /dev/null
+++ b/Engine/Plugins/Experimental/GameplayGraph/Source/GameplayGraph/Public/Graph/Graph.h
@@ -0,0 +1,158 @@
+// Copyright Epic Games, Inc. All Rights Reserved.
+#pragma once
+
+#include "UObject/Object.h"
+
+#include "Containers/Array.h"
+#include "Containers/Map.h"
+#include "Containers/Set.h"
+#include "Graph/GraphEdge.h"
+#include "Graph/GraphHandle.h"
+#include "Graph/GraphIsland.h"
+#include "Graph/GraphVertex.h"
+#include "Templates/SubclassOf.h"
+
+#include "Graph.generated.h"
+
+USTRUCT()
+struct FGraphProperties
+{
+	GENERATED_BODY()
+
+	UPROPERTY(SaveGame)
+	bool bGenerateIslands = true;
+};
+
+/**
+ * The minimum amount of data we need to serialize to be able to reconstruct the graph as it was.
+ * Note that classes that inherit from UGraph and its elements will no doubt want to extend the graph
+ * with actual information on each node/edge/island. In that case, they should extend FSerializableGraph
+ * to contain the extra information per graph handle. Furthermore, they'll need to extend UGraph to have
+ * its own typed serialization save/load functions that call the base functions in UGraph first.
+ */
+USTRUCT()
+struct FSerializableGraph
+{
+	GENERATED_BODY()
+
+	UPROPERTY(SaveGame)
+	FGraphProperties Properties;
+
+	UPROPERTY(SaveGame)
+	TArray<FGraphVertexHandle> Vertices;
+
+	UPROPERTY(SaveGame)
+	TMap<FGraphEdgeHandle, FSerializedEdgeData> Edges;
+
+	UPROPERTY(SaveGame)
+	TMap<FGraphIslandHandle, FSerializedIslandData> Islands;
+};
+
+/**
+ * A UGraph is a collection of nodes and edges. This graph representation
+ * is meant to be easily integrable into gameplay systems in the Unreal Engine.
+ * 
+ * Conceptually, you can imagine that a graph is meant to easily represent relationships
+ * so we can answer queries such as:
+ *	- Are these two nodes connected to each other?
+ *  - How far away are these two nodes?
+ *  - Who is the closest node that has XYZ?
+ *  - etc.
+ * 
+ * UGraph provides an interface to be able to run such queries. However, ultimately what
+ * makes the graph useful is not only the relationships represented by the edges, but also
+ * the data that is stored on each node and each edge. Depending on what the user wants to
+ * represent, the user will have to subclass UGraphVertex and UGraphEdge to hold that data.
+ * 
+ * As the user adds nodes and edges into the graph, they will also be implicitly creating "islands"
+ * (i.e. a connected component in the graph). Each graph may have multiple islands. Users
+ * of the graph can disable the island detection/creation if needed.
+ * 
+ * Note that this is an UNDIRECTED GRAPH.
+ */
+UCLASS()
+class GAMEPLAYGRAPH_API UGraph: public UObject
+{
+	GENERATED_BODY()
+public:
+	UGraph() = default;
+
+	FSerializableGraph GetSerializableGraph() const;
+	void LoadFromSerializedGraph(const FSerializableGraph& Input);
+
+	void Empty();
+	void InitializeFromProperties(const FGraphProperties& Properties);
+
+	/** Given a node handle, find the handle in the current graph with a proper element set. */
+	FGraphVertexHandle GetCompleteNodeHandle(const FGraphVertexHandle& InHandle);
+
+	/** Create a node with the specified subclass, adds it to the graph, and returns a handle to it. */
+	FGraphVertexHandle CreateVertex(int64 InUniqueIndex = INDEX_NONE);
+
+	/** Creates an edge between the two nodes. */
+	FGraphEdgeHandle CreateEdge(FGraphVertexHandle Node1, FGraphVertexHandle Node2, int64 InUniqueIndex = INDEX_NONE, bool bAddToIslands = true);
+
+	/** Removes a node from the graph along with any edges that contain it. */
+	void RemoveVertex(const FGraphVertexHandle& NodeHandle);
+
+	/** Removes an edge from the graph. */
+	void RemoveEdge(const FGraphEdgeHandle& EdgeHandle, bool bHandleIslands);
+
+	/** This should be called immediately after a node and any relevant edges have been added to the graph. */
+	virtual void FinalizeVertex(const FGraphVertexHandle& InHandle);
+
+	/** Remove an island from the graph. */
+	void RemoveIsland(const FGraphIslandHandle& IslandHandle);
+protected:
+	const TMap<FGraphVertexHandle, TObjectPtr<UGraphVertex>>& GetVertices() const { return Vertices; }
+	const TMap<FGraphEdgeHandle, TObjectPtr<UGraphEdge>>& GetEdges() const { return Edges; }
+	const TMap<FGraphIslandHandle, TObjectPtr<UGraphIsland>>& GetIslands() const { return Islands; }
+	const FGraphProperties& GetProperties() const { return Properties; }
+
+private:
+	UPROPERTY()
+	TMap<FGraphVertexHandle, TObjectPtr<UGraphVertex>> Vertices;
+
+	/** Keeps track of every edge a node is a part of. The node itself will only track edges it's the "source" node of in directed graphs. */
+	TMap<FGraphVertexHandle, TSet<FGraphEdgeHandle>> VertexEdges;
+
+	UPROPERTY()
+	TMap<FGraphEdgeHandle, TObjectPtr<UGraphEdge>> Edges;
+
+	UPROPERTY()
+	TMap<FGraphIslandHandle, TObjectPtr<UGraphIsland>> Islands;
+	FGraphProperties Properties;
+
+	/**
+	 *  The unique index to assign to the next node that gets created.
+	 *  If loading from persistence, as nodes get added into the graph, this
+	 *  value will automatically keep increasing so we never try to reuse
+	 *  the same unique index.
+	 */
+	int64 NextAvailableVertexUniqueIndex = 0;
+	int64 NextAvailableEdgeUniqueIndex = 0;
+	int64 NextAvailableIslandUniqueIndex = 0;
+
+	/** Creates an island out of a given set of nodes. */
+	FGraphIslandHandle CreateIsland(TArray<FGraphVertexHandle> Nodes, int64 InUniqueIndex = INDEX_NONE);
+
+	/** Adds a node to the graph's node collection and modifies the NextAvailableNodeUniqueIndex as necessary to maintain its validity. */
+	void RegisterVertex(TObjectPtr<UGraphVertex> Node);
+
+	/** Add an edge to the graph and modifies the NextAvailableEdgeUniqueIndex as necessary to maintain its validity. */
+	void RegisterEdge(TObjectPtr<UGraphEdge> Edge);
+
+	/** Add an island to the graph and modifies the NextAvailableIslandUniqueIndex as necessary to maintain its validity. */
+	void RegisterIsland(TObjectPtr<UGraphIsland> Edge);
+
+	/** When we add an edge into a graph, we will want to merge islands if the edge connects two nodes that are in different islands. */
+	void MergeOrCreateIslands(TObjectPtr<UGraphEdge> Edge);
+
+	/** After a change, this function will remove the island if it's empty or will attempt to split it into two smaller islands. */
+	void RemoveOrSplitIsland(TObjectPtr<UGraphIsland> Island);
+
+	/** Factory functions for creating an appropriately typed node/edge/island. */
+	virtual TObjectPtr<UGraphVertex> CreateTypedVertex() const;
+	virtual TObjectPtr<UGraphEdge> CreateTypedEdge() const;
+	virtual TObjectPtr<UGraphIsland> CreateTypedIsland() const;
+};
\ No newline at end of file
diff --git a/Engine/Plugins/Experimental/GameplayGraph/Source/GameplayGraph/Public/Graph/GraphEdge.h b/Engine/Plugins/Experimental/GameplayGraph/Source/GameplayGraph/Public/Graph/GraphEdge.h
new file mode 100644
--- /dev/null
+++ b/Engine/Plugins/Experimental/GameplayGraph/Source/GameplayGraph/Public/Graph/GraphEdge.h
@@ -0,0 +1,55 @@
+// Copyright Epic Games, Inc. All Rights Reserved.
+#pragma once
+
+#include "Graph/GraphElement.h"
+#include "GraphEdge.generated.h"
+
+class UGraphVertex;
+
+USTRUCT()
+struct FSerializedEdgeData
+{
+	GENERATED_BODY()
+
+	UPROPERTY(SaveGame)
+	FGraphVertexHandle Node1;
+
+	UPROPERTY(SaveGame)
+	FGraphVertexHandle Node2;
+};
+
+UCLASS()
+class GAMEPLAYGRAPH_API UGraphEdge : public UGraphElement
+{
+	GENERATED_BODY()
+public:
+	UGraphEdge();
+
+	const FGraphVertexHandle& GetOtherNode(const FGraphVertexHandle& InNode) const;
+	bool ContainsNode(const FGraphVertexHandle& InNode) const;
+
+	const FGraphVertexHandle& NodeA() const { return A; }
+	const FGraphVertexHandle& NodeB() const { return B; }
+
+	FGraphVertexHandle& NodeA() { return A; }
+	FGraphVertexHandle& NodeB() { return B; }
+
+	FGraphEdgeHandle Handle() const
+	{
+		return FGraphEdgeHandle{ GetUniqueIndex(), const_cast<UGraphEdge*>(this) };
+	}
+
+	FSerializedEdgeData GetSerializedData() const;
+
+	friend class UGraph;
+
+protected:
+	void SetNodes(const FGraphVertexHandle& InA, const FGraphVertexHandle& InB);
+
+private:
+	UPROPERTY(SaveGame)
+	FGraphVertexHandle A;
+
+	UPROPERTY(SaveGame)
+	FGraphVertexHandle B;
+};
\ No newline at end of file
diff --git a/Engine/Plugins/Experimental/GameplayGraph/Source/GameplayGraph/Public/Graph/GraphElement.h b/Engine/Plugins/Experimental/GameplayGraph/Source/GameplayGraph/Public/Graph/GraphElement.h
new file mode 100644
--- /dev/null
+++ b/Engine/Plugins/Experimental/GameplayGraph/Source/GameplayGraph/Public/Graph/GraphElement.h
@@ -0,0 +1,52 @@
+// Copyright Epic Games, Inc. All Rights Reserved.
+#pragma once
+
+#include "UObject/Object.h"
+#include "Graph/GraphHandle.h"
+
+#include "GraphElement.generated.h"
+
+UENUM()
+enum class EGraphElementType
+{
+	Node,
+	Edge,
+	Island,
+	Unknown
+};
+
+UCLASS(abstract)
+class UGraphElement : public UObject
+{
+	GENERATED_BODY()
+public:
+	explicit UGraphElement(EGraphElementType InElementType)
+		: ElementType(InElementType)
+	{}
+
+	EGraphElementType GetElementType() const { return ElementType; }
+
+	friend class UGraph;
+protected:
+	UGraphElement() = default;
+
+	void SetUniqueIndex(int64 InUniqueIndex) { UniqueIndex = InUniqueIndex; }
+	int64 GetUniqueIndex() const { return UniqueIndex; }
+
+	void SetParentGraph(TObjectPtr<UGraph> InGraph) { ParentGraph = InGraph; }
+	TObjectPtr<UGraph> GetGraph() const { return ParentGraph.Get(); }
+
+	/** Called when we create this element and prior to setting any properties. */
+	virtual void OnCreate() {}
+
+private:
+	UPROPERTY()
+	EGraphElementType ElementType = EGraphElementType::Unknown;
+
+	/** Will match the UniqueIndex in the UGraphHandle that references this element. */
+	UPROPERTY()
+	int64 UniqueIndex = INDEX_NONE;
+
+	UPROPERTY()
+	TWeakObjectPtr<UGraph> ParentGraph = nullptr;
+};
\ No newline at end of file
diff --git a/Engine/Plugins/Experimental/GameplayGraph/Source/GameplayGraph/Public/Graph/GraphHandle.h b/Engine/Plugins/Experimental/GameplayGraph/Source/GameplayGraph/Public/Graph/GraphHandle.h
new file mode 100644
--- /dev/null
+++ b/Engine/Plugins/Experimental/GameplayGraph/Source/GameplayGraph/Public/Graph/GraphHandle.h
@@ -0,0 +1,89 @@
+// Copyright Epic Games, Inc. All Rights Reserved.
+#pragma once
+
+#include "UObject/Object.h"
+#include "GraphHandle.generated.h"
+
+class UGraphElement;
+class UGraphVertex;
+class UGraphEdge;
+class UGraphIsland;
+
+/**
+ * For persistence, every node in a graph is given a unique index.
+ * A FGraphHandle encapsulates that index to make it easy to go from
+ * the index to the node and vice versa.
+ */
+USTRUCT()
+struct GAMEPLAYGRAPH_API FGraphHandle
+{
+	GENERATED_BODY()
+public:
+	FGraphHandle();
+	FGraphHandle(int64 InUniqueIndex, TObjectPtr<UGraphElement> InElement);
+
+	void Clear();
+
+	/** Whether or not this handle has been initialized. */
+	bool IsValid() const;
+	bool HasElement() const;
+	bool IsComplete() const;
+
+	int64 GetUniqueIndex() const { return UniqueIndex; }
+
+	void SetElement(TObjectPtr<UGraphElement> InElement);
+	TObjectPtr<UGraphElement> GetElement() const;
+
+	bool operator==(const FGraphHandle& Other) const;
+	bool operator!=(const FGraphHandle& Other) const;
+	bool operator<(const FGraphHandle& Other) const;
+
+	friend uint32 GAMEPLAYGRAPH_API GetTypeHash(const FGraphHandle& Handle);
+private:
+	/** Unique identifier within a graph. */
+	UPROPERTY(SaveGame)
+	int64 UniqueIndex = INDEX_NONE;
+
+	/** Pointer to the graph */
+	UPROPERTY(Transient)
+	TWeakObjectPtr<UGraphElement> Element;
+};
+
+USTRUCT()
+struct GAMEPLAYGRAPH_API FGraphVertexHandle : public FGraphHandle
+{
+	GENERATED_BODY()
+
+	FGraphVertexHandle();
+	FGraphVertexHandle(int64 InUniqueIndex, TObjectPtr<UGraphElement> InElement)
+		: FGraphHandle(InUniqueIndex, InElement)
+	{}
+
+	TObjectPtr<UGraphVertex> GetVertex() const;
+};
+
+USTRUCT()
+struct GAMEPLAYGRAPH_API FGraphEdgeHandle : public FGraphHandle
+{
+	GENERATED_BODY()
+
+	FGraphEdgeHandle();
+	FGraphEdgeHandle(int64 InUniqueIndex, TObjectPtr<UGraphElement> InElement)
+		: FGraphHandle(InUniqueIndex, InElement)
+	{}
+
+	TObjectPtr<UGraphEdge> GetEdge() const;
+};
+
+USTRUCT()
+struct GAMEPLAYGRAPH_API FGraphIslandHandle : public FGraphHandle
+{
+	GENERATED_BODY()
+
+	FGraphIslandHandle();
+	FGraphIslandHandle(int64 InUniqueIndex, TObjectPtr<UGraphElement> InElement)
+		: FGraphHandle(InUniqueIndex, InElement)
+	{}
+
+	TObjectPtr<UGraphIsland> GetIsland() const;
+};
\ No newline at end of file
diff --git a/Engine/Plugins/Experimental/GameplayGraph/Source/GameplayGraph/Public/Graph/GraphIsland.h b/Engine/Plugins/Experimental/GameplayGraph/Source/GameplayGraph/Public/Graph/GraphIsland.h
new file mode 100644
--- /dev/null
+++ b/Engine/Plugins/Experimental/GameplayGraph/Source/GameplayGraph/Public/Graph/GraphIsland.h
@@ -0,0 +1,76 @@
+// Copyright Epic Games, Inc. All Rights Reserved.
+#pragma once
+
+#include "Graph/GraphElement.h"
+#include "GraphIsland.generated.h"
+
+USTRUCT()
+struct FSerializedIslandData
+{
+	GENERATED_BODY()
+
+	UPROPERTY(SaveGame)
+	TArray<FGraphVertexHandle> Vertices;
+};
+
+/** Delegate to track when some sort of batch change has occurred on this island that probably changes its connectivity.
+ *  This is different from FOnGraphIslandNodeRemoved since FOnGraphIslandDestructiveChangeFinish will only be called
+ *  once for the graph for a given operation while FOnGraphIslandNodeRemoved may be called multiple times if we're removing
+ *  more than one node from the island at a given time. Note that this will only be called as a result of a destructive change.
+ *  So repeatedly adding a node to an island won't call this event. */
+DECLARE_MULTICAST_DELEGATE(FOnGraphIslandConnectedComponentsChanged);
+
+/** Delegate to track when this island should no longer exist. */
+DECLARE_MULTICAST_DELEGATE(FOnGraphIslandDestroyed);
+
+/** Delegate to track the event when the island has a node added to it. */
+DECLARE_MULTICAST_DELEGATE_OneParam(FOnGraphIslandVertexAdded, const FGraphVertexHandle&);
+
+/** Delegate to track the event when the island has a node removed from it. */
+DECLARE_MULTICAST_DELEGATE_OneParam(FOnGraphIslandVertexRemoved, const FGraphVertexHandle&);
+
+UCLASS()
+class GAMEPLAYGRAPH_API UGraphIsland : public UGraphElement
+{
+	GENERATED_BODY()
+public:
+	UGraphIsland();
+
+	FGraphIslandHandle Handle() const
+	{
+		return FGraphIslandHandle{ GetUniqueIndex(), const_cast<UGraphIsland*>(this) };
+	}
+
+	bool IsEmpty() const { return Vertices.IsEmpty(); }
+	const TSet<FGraphVertexHandle>& GetVertices() const { return Vertices; }
+	FSerializedIslandData GetSerializedData() const;
+
+	FOnGraphIslandVertexAdded OnVertexAdded;
+	FOnGraphIslandVertexRemoved OnVertexRemoved;
+	FOnGraphIslandDestroyed OnDestroyed;
+	FOnGraphIslandConnectedComponentsChanged OnConnectivityChanged;
+
+	friend class UGraph;
+protected:
+
+	virtual void HandleOnVertexAdded(const FGraphVertexHandle& Handle);
+	virtual void HandleOnVertexRemoved(const FGraphVertexHandle& Handle);
+	virtual void HandleOnDestroyed();
+	virtual void HandleOnConnectivityChanged();
+
+	/** Called when removing the island from the graph. */
+	void Destroy();
+
+	/** Takes the nodes from the other island and puts them into this island. */
+	void MergeWith(TObjectPtr<UGraphIsland> OtherIsland);
+
+	/** Adds a single node into this island. */
+	void AddVertex(const FGraphVertexHandle& Node);
+
+	/** Removes a node from the island. */
+	void RemoveVertex(const FGraphVertexHandle& Node);
+
+private:
+	UPROPERTY(SaveGame)
+	TSet<FGraphVertexHandle> Vertices;
+};
\ No newline at end of file
diff --git a/Engine/Plugins/Experimental/GameplayGraph/Source/GameplayGraph/Public/Graph/GraphVertex.h b/Engine/Plugins/Experimental/GameplayGraph/Source/GameplayGraph/Public/Graph/GraphVertex.h
new file mode 100644
--- /dev/null
+++ b/Engine/Plugins/Experimental/GameplayGraph/Source/GameplayGraph/Public/Graph/GraphVertex.h
@@ -0,0 +1,53 @@
+// Copyright Epic Games, Inc. All Rights Reserved.
+#pragma once
+
+#include "Graph/GraphElement.h"
+
+#include "GraphVertex.generated.h"
+
+/** Event for when the node has been removed from the graph. */
+DECLARE_MULTICAST_DELEGATE(FOnGraphVertexRemoved);
+
+UCLASS()
+class GAMEPLAYGRAPH_API UGraphVertex : public UGraphElement
+{
+	GENERATED_BODY()
+public:
+	UGraphVertex();
+
+	FGraphVertexHandle Handle() const
+	{
+		return FGraphVertexHandle{ GetUniqueIndex(), const_cast<UGraphVertex*>(this)};
+	}
+
+	bool HasEdgeTo(const FGraphVertexHandle& Other) const;
+	const FGraphIslandHandle& GetParentIsland() const { return ParentIsland; }
+
+	template<typename TLambda>
+	void ForEachAdjacentVertex(TLambda&& Lambda)
+	{
+		for (const TPair<FGraphVertexHandle, FGraphEdgeHandle>& Kvp : Edges)
+		{
+			Lambda(Kvp.Key);
+		}
+	}
+
+	FOnGraphVertexRemoved OnVertexRemoved;
+
+	friend class UGraph;
+	friend class UGraphIsland;
+protected:
+
+	void AddEdgeTo(const FGraphVertexHandle& Node, const FGraphEdgeHandle& Edge);
+	void RemoveEdge(const FGraphEdgeHandle& EdgeHandle);
+	void SetParentIsland(const FGraphIslandHandle& Island) { ParentIsland = Island; }
+
+	virtual void HandleOnVertexRemoved();
+private:
+
+	UPROPERTY()
+	TMap<FGraphVertexHandle, FGraphEdgeHandle> Edges;
+
+	UPROPERTY()
+	FGraphIslandHandle ParentIsland;
+};
\ No newline at end of file
