fix non-unity build for the gameplay graph plugin

diff --git a/Engine/Plugins/Experimental/GameplayGraph/Source/GameplayGraph/Private/Graph/GraphElement.cpp b/Engine/Plugins/Experimental/GameplayGraph/Source/GameplayGraph/Private/Graph/GraphElement.cpp
new file mode 100644
--- /dev/null
+++ b/Engine/Plugins/Experimental/GameplayGraph/Source/GameplayGraph/Private/Graph/GraphElement.cpp
@@ -0,0 +1,23 @@
+// Copyright Epic Games, Inc. All Rights Reserved.
+
+#include "Graph/GraphElement.h"
+#include "Graph/Graph.h"
+
+UGraphElement::UGraphElement(EGraphElementType InElementType)
+	: ElementType(InElementType)
+{
+}
+
+UGraphElement::UGraphElement()
+{
+}
+
+void UGraphElement::SetParentGraph(TObjectPtr<UGraph> InGraph)
+{
+	ParentGraph = InGraph;
+}
+
+TObjectPtr<UGraph> UGraphElement::GetGraph() const
+{
+	return ParentGraph.Get();
+}
\ No newline at end of file
diff --git a/Engine/Plugins/Experimental/GameplayGraph/Source/GameplayGraph/Public/Graph/GraphElement.h b/Engine/Plugins/Experimental/GameplayGraph/Source/GameplayGraph/Public/Graph/GraphElement.h
--- a/Engine/Plugins/Experimental/GameplayGraph/Source/GameplayGraph/Public/Graph/GraphElement.h
+++ b/Engine/Plugins/Experimental/GameplayGraph/Source/GameplayGraph/Public/Graph/GraphElement.h
@@ -6,6 +6,8 @@
 
 #include "GraphElement.generated.h"
 
+class UGraph;
+
 UENUM()
 enum class EGraphElementType
 {
@@ -16,25 +18,22 @@ enum class EGraphElementType
 };
 
 UCLASS(abstract)
-class UGraphElement : public UObject
+class GAMEPLAYGRAPH_API UGraphElement : public UObject
 {
 	GENERATED_BODY()
 public:
-	explicit UGraphElement(EGraphElementType InElementType)
-		: ElementType(InElementType)
-	{}
-
+	explicit UGraphElement(EGraphElementType InElementType);
 	EGraphElementType GetElementType() const { return ElementType; }
 
 	friend class UGraph;
 protected:
-	UGraphElement() = default;
+	UGraphElement();
 
 	void SetUniqueIndex(int64 InUniqueIndex) { UniqueIndex = InUniqueIndex; }
 	int64 GetUniqueIndex() const { return UniqueIndex; }
 
-	void SetParentGraph(TObjectPtr<UGraph> InGraph) { ParentGraph = InGraph; }
-	TObjectPtr<UGraph> GetGraph() const { return ParentGraph.Get(); }
+	void SetParentGraph(TObjectPtr<UGraph> InGraph);
+	TObjectPtr<UGraph> GetGraph() const;
 
 	/** Called when we create this element and prior to setting any properties. */
 	virtual void OnCreate() {}
@@ -48,5 +47,5 @@ private:
 	int64 UniqueIndex = INDEX_NONE;
 
 	UPROPERTY()
-	TWeakObjectPtr<UGraph> ParentGraph = nullptr;
+	TWeakObjectPtr<UGraph> ParentGraph;
 };
\ No newline at end of file
