fix additional CIS errors for non-unity builds in the gameplay graph plugin

diff --git a/Engine/Plugins/Experimental/GameplayGraph/Source/GameplayGraph/Private/Graph/Algorithms/Search/Search.cpp b/Engine/Plugins/Experimental/GameplayGraph/Source/GameplayGraph/Private/Graph/Algorithms/Search/Search.cpp
--- a/Engine/Plugins/Experimental/GameplayGraph/Source/GameplayGraph/Private/Graph/Algorithms/Search/Search.cpp
+++ b/Engine/Plugins/Experimental/GameplayGraph/Source/GameplayGraph/Private/Graph/Algorithms/Search/Search.cpp
@@ -13,10 +13,10 @@ namespace Graph::Algorithms
 		using BFSDataStructure = TQueue<FGraphVertexHandle>;
 		using DFSDataStructure = TArray<FGraphVertexHandle>;
 
-		template<template<typename> typename TDataStructure>
-		FGraphVertexHandle GetNextAndAdvance(TDataStructure<FGraphVertexHandle>& Queue)
+		template<typename TDataStructure>
+		FGraphVertexHandle GetNextAndAdvance(TDataStructure& Queue)
 		{
-			if constexpr (std::is_same_v<TDataStructure<FGraphVertexHandle>, BFSDataStructure>)
+			if constexpr (std::is_same_v<TDataStructure, BFSDataStructure>)
 			{
 				FGraphVertexHandle Ret;
 				if (ensure(Queue.Dequeue(Ret)))
@@ -25,7 +25,7 @@ namespace Graph::Algorithms
 				}
 				return {};
 			}
-			else if constexpr (std::is_same_v<TDataStructure<FGraphVertexHandle>, DFSDataStructure>)
+			else if constexpr (std::is_same_v<TDataStructure, DFSDataStructure>)
 			{
 				if (ensure(!Queue.IsEmpty()))
 				{
@@ -42,14 +42,14 @@ namespace Graph::Algorithms
 			}
 		}
 
-		template<template<typename> typename TDataStructure>
-		void AddToWorkQueue(TDataStructure<FGraphVertexHandle>& Queue, const FGraphVertexHandle& Handle)
+		template<typename TDataStructure>
+		void AddToWorkQueue(TDataStructure& Queue, const FGraphVertexHandle& Handle)
 		{
-			if constexpr (std::is_same_v<TDataStructure<FGraphVertexHandle>, BFSDataStructure>)
+			if constexpr (std::is_same_v<TDataStructure, BFSDataStructure>)
 			{
 				Queue.Enqueue(Handle);
 			}
-			else if constexpr (std::is_same_v<TDataStructure<FGraphVertexHandle>, DFSDataStructure>)
+			else if constexpr (std::is_same_v<TDataStructure, DFSDataStructure>)
 			{
 				Queue.Add(Handle);
 			}
@@ -59,10 +59,10 @@ namespace Graph::Algorithms
 			}
 		}
 
-		template<template<typename> typename TDataStructure>
+		template<typename TDataStructure>
 		FGraphVertexHandle GenericSearch(const FGraphVertexHandle& Start, FSearchCallback Callback)
 		{
-			TDataStructure<FGraphVertexHandle> WorkQueue;
+			TDataStructure WorkQueue;
 			AddToWorkQueue<TDataStructure>(WorkQueue, Start);
 
 			TSet<FGraphVertexHandle> Seen;
@@ -104,11 +104,11 @@ namespace Graph::Algorithms
 
 	FGraphVertexHandle BFS(const FGraphVertexHandle& Start, FSearchCallback Callback)
 	{
-		return GenericSearch<TQueue>(Start, Callback);
+		return GenericSearch<BFSDataStructure>(Start, Callback);
 	}
 
 	FGraphVertexHandle DFS(const FGraphVertexHandle& Start, FSearchCallback Callback)
 	{
-		return GenericSearch<TArray>(Start, Callback);
+		return GenericSearch<DFSDataStructure>(Start, Callback);
 	}
 }
\ No newline at end of file
diff --git a/Engine/Plugins/Experimental/GameplayGraph/Source/GameplayGraph/Private/Graph/GraphHandle.cpp b/Engine/Plugins/Experimental/GameplayGraph/Source/GameplayGraph/Private/Graph/GraphHandle.cpp
--- a/Engine/Plugins/Experimental/GameplayGraph/Source/GameplayGraph/Private/Graph/GraphHandle.cpp
+++ b/Engine/Plugins/Experimental/GameplayGraph/Source/GameplayGraph/Private/Graph/GraphHandle.cpp
@@ -3,6 +3,7 @@
 #include "Graph/GraphHandle.h"
 #include "Graph/GraphElement.h"
 #include "Graph/GraphEdge.h"
+#include "Graph/GraphIsland.h"
 #include "Graph/GraphVertex.h"
 
 #include UE_INLINE_GENERATED_CPP_BY_NAME(GraphHandle)
