add ability to query some numeric size properties on gameplay graphs

diff --git a/Engine/Plugins/Experimental/GameplayGraph/Source/GameplayGraph/Public/Graph/Graph.h b/Engine/Plugins/Experimental/GameplayGraph/Source/GameplayGraph/Public/Graph/Graph.h
--- a/Engine/Plugins/Experimental/GameplayGraph/Source/GameplayGraph/Public/Graph/Graph.h
+++ b/Engine/Plugins/Experimental/GameplayGraph/Source/GameplayGraph/Public/Graph/Graph.h
@@ -103,6 +103,11 @@ public:
 
 	/** Remove an island from the graph. */
 	void RemoveIsland(const FGraphIslandHandle& IslandHandle);
+
+	int32 NumVertices() const { return Vertices.Num(); }
+	int32 NumEdges() const { return Edges.Num(); }
+	int32 NumIslands() const { return Islands.Num(); }
+
 protected:
 	const TMap<FGraphVertexHandle, TObjectPtr<UGraphVertex>>& GetVertices() const { return Vertices; }
 	const TMap<FGraphEdgeHandle, TObjectPtr<UGraphEdge>>& GetEdges() const { return Edges; }
