every vertex in the graph is added to an island even if it has no edges

diff --git a/Engine/Plugins/Experimental/GameplayGraph/Source/GameplayGraph/Private/Graph/Graph.cpp b/Engine/Plugins/Experimental/GameplayGraph/Source/GameplayGraph/Private/Graph/Graph.cpp
--- a/Engine/Plugins/Experimental/GameplayGraph/Source/GameplayGraph/Private/Graph/Graph.cpp
+++ b/Engine/Plugins/Experimental/GameplayGraph/Source/GameplayGraph/Private/Graph/Graph.cpp
@@ -498,10 +498,11 @@ void UGraph::FinalizeVertex(const FGraphVertexHandle& InHandle)
 		return;
 	}
 
-	const FGraphIslandHandle& IslandHandle = Node->GetParentIsland();
+	FGraphIslandHandle IslandHandle = Node->GetParentIsland();
 	if (!IslandHandle.IsComplete())
 	{
-		return;
+		// Every vertex must be in an island.
+		IslandHandle = CreateIsland({ InHandle });
 	}
 
 	if (TObjectPtr<UGraphIsland> Island = IslandHandle.GetIsland())
