graph island delegates should also be passing the graph island's handle as well

diff --git a/Engine/Plugins/Experimental/GameplayGraph/Source/GameplayGraph/Private/Graph/GraphIsland.cpp b/Engine/Plugins/Experimental/GameplayGraph/Source/GameplayGraph/Private/Graph/GraphIsland.cpp
--- a/Engine/Plugins/Experimental/GameplayGraph/Source/GameplayGraph/Private/Graph/GraphIsland.cpp
+++ b/Engine/Plugins/Experimental/GameplayGraph/Source/GameplayGraph/Private/Graph/GraphIsland.cpp
@@ -93,20 +93,20 @@ FSerializedIslandData UGraphIsland::GetSerializedData() const
 
 void UGraphIsland::HandleOnVertexAdded(const FGraphVertexHandle& Handle)
 {
-	OnVertexAdded.Broadcast(Handle);
+	OnVertexAdded.Broadcast(this->Handle(), Handle);
 }
 
 void UGraphIsland::HandleOnVertexRemoved(const FGraphVertexHandle& Handle)
 {
-	OnVertexRemoved.Broadcast(Handle);
+	OnVertexRemoved.Broadcast(this->Handle(), Handle);
 }
 
 void UGraphIsland::HandleOnDestroyed()
 {
-	OnDestroyed.Broadcast();
+	OnDestroyed.Broadcast(Handle());
 }
 
 void UGraphIsland::HandleOnConnectivityChanged()
 {
-	OnConnectivityChanged.Broadcast();
+	OnConnectivityChanged.Broadcast(Handle());
 }
\ No newline at end of file
diff --git a/Engine/Plugins/Experimental/GameplayGraph/Source/GameplayGraph/Public/Graph/GraphIsland.h b/Engine/Plugins/Experimental/GameplayGraph/Source/GameplayGraph/Public/Graph/GraphIsland.h
--- a/Engine/Plugins/Experimental/GameplayGraph/Source/GameplayGraph/Public/Graph/GraphIsland.h
+++ b/Engine/Plugins/Experimental/GameplayGraph/Source/GameplayGraph/Public/Graph/GraphIsland.h
@@ -18,16 +18,16 @@ struct FSerializedIslandData
  *  once for the graph for a given operation while FOnGraphIslandNodeRemoved may be called multiple times if we're removing
  *  more than one node from the island at a given time. Note that this will only be called as a result of a destructive change.
  *  So repeatedly adding a node to an island won't call this event. */
-DECLARE_MULTICAST_DELEGATE(FOnGraphIslandConnectedComponentsChanged);
+DECLARE_MULTICAST_DELEGATE_OneParam(FOnGraphIslandConnectedComponentsChanged, const FGraphIslandHandle&);
 
 /** Delegate to track when this island should no longer exist. */
-DECLARE_MULTICAST_DELEGATE(FOnGraphIslandDestroyed);
+DECLARE_MULTICAST_DELEGATE_OneParam(FOnGraphIslandDestroyed, const FGraphIslandHandle&);
 
 /** Delegate to track the event when the island has a node added to it. */
-DECLARE_MULTICAST_DELEGATE_OneParam(FOnGraphIslandVertexAdded, const FGraphVertexHandle&);
+DECLARE_MULTICAST_DELEGATE_TwoParams(FOnGraphIslandVertexAdded, const FGraphIslandHandle&, const FGraphVertexHandle&);
 
 /** Delegate to track the event when the island has a node removed from it. */
-DECLARE_MULTICAST_DELEGATE_OneParam(FOnGraphIslandVertexRemoved, const FGraphVertexHandle&);
+DECLARE_MULTICAST_DELEGATE_TwoParams(FOnGraphIslandVertexRemoved, const FGraphIslandHandle&, const FGraphVertexHandle&);
 
 UCLASS()
 class GAMEPLAYGRAPH_API UGraphIsland : public UGraphElement
