add delegates for creation of different graph elements

diff --git a/Engine/Plugins/Experimental/GameplayGraph/Source/GameplayGraph/Private/Graph/Graph.cpp b/Engine/Plugins/Experimental/GameplayGraph/Source/GameplayGraph/Private/Graph/Graph.cpp
--- a/Engine/Plugins/Experimental/GameplayGraph/Source/GameplayGraph/Private/Graph/Graph.cpp
+++ b/Engine/Plugins/Experimental/GameplayGraph/Source/GameplayGraph/Private/Graph/Graph.cpp
@@ -119,6 +119,7 @@ FGraphVertexHandle UGraph::CreateVertex(int64 InUniqueIndex)
 	// No need to increment since AddNode will do that.
 	Vertex->SetUniqueIndex((InUniqueIndex == INDEX_NONE) ? NextAvailableVertexUniqueIndex : InUniqueIndex);
 	RegisterVertex(Vertex);
+	OnVertexCreated.Broadcast(Vertex->Handle());
 	return Vertex->Handle();
 }
 
@@ -177,6 +178,8 @@ FGraphEdgeHandle UGraph::CreateEdge(FGraphVertexHandle Node1, FGraphVertexHandle
 	Node1Ptr->AddEdgeTo(Node2, Edge->Handle());
 	Node2Ptr->AddEdgeTo(Node1, Edge->Handle());
 
+	OnEdgeCreated.Broadcast(Edge->Handle());
+
 	// If we want to keep track of islands, this is where we need to create/merge islands.
 	if (Properties.bGenerateIslands && bAddToIslands)
 	{
@@ -224,6 +227,7 @@ FGraphIslandHandle UGraph::CreateIsland(TArray<FGraphVertexHandle> InputNodes, i
 	}
 
 	RegisterIsland(Island);
+	OnIslandCreated.Broadcast(Island->Handle());
 	return Island->Handle();
 }
 
diff --git a/Engine/Plugins/Experimental/GameplayGraph/Source/GameplayGraph/Public/Graph/Graph.h b/Engine/Plugins/Experimental/GameplayGraph/Source/GameplayGraph/Public/Graph/Graph.h
--- a/Engine/Plugins/Experimental/GameplayGraph/Source/GameplayGraph/Public/Graph/Graph.h
+++ b/Engine/Plugins/Experimental/GameplayGraph/Source/GameplayGraph/Public/Graph/Graph.h
@@ -23,6 +23,10 @@ struct FGraphProperties
 	bool bGenerateIslands = true;
 };
 
+DECLARE_MULTICAST_DELEGATE_OneParam(FOnGraphVertexCreated, const FGraphVertexHandle&);
+DECLARE_MULTICAST_DELEGATE_OneParam(FOnGraphEdgeCreated, const FGraphEdgeHandle&);
+DECLARE_MULTICAST_DELEGATE_OneParam(FOnGraphIslandCreated, const FGraphIslandHandle&);
+
 /**
  * The minimum amount of data we need to serialize to be able to reconstruct the graph as it was.
  * Note that classes that inherit from UGraph and its elements will no doubt want to extend the graph
@@ -108,6 +112,10 @@ public:
 	int32 NumEdges() const { return Edges.Num(); }
 	int32 NumIslands() const { return Islands.Num(); }
 
+	FOnGraphVertexCreated OnVertexCreated;
+	FOnGraphEdgeCreated OnEdgeCreated;
+	FOnGraphIslandCreated OnIslandCreated;
+
 protected:
 	const TMap<FGraphVertexHandle, TObjectPtr<UGraphVertex>>& GetVertices() const { return Vertices; }
 	const TMap<FGraphEdgeHandle, TObjectPtr<UGraphEdge>>& GetEdges() const { return Edges; }
