merging graph islands should first remove the vertex from the old island

diff --git a/Engine/Plugins/Experimental/GameplayGraph/Source/GameplayGraph/Private/Graph/GraphIsland.cpp b/Engine/Plugins/Experimental/GameplayGraph/Source/GameplayGraph/Private/Graph/GraphIsland.cpp
--- a/Engine/Plugins/Experimental/GameplayGraph/Source/GameplayGraph/Private/Graph/GraphIsland.cpp
+++ b/Engine/Plugins/Experimental/GameplayGraph/Source/GameplayGraph/Private/Graph/GraphIsland.cpp
@@ -39,8 +39,10 @@ void UGraphIsland::MergeWith(TObjectPtr<UGraphIsland> OtherIsland)
 		return;
 	}
 
-	for (const FGraphVertexHandle& Node : OtherIsland->Vertices)
+	TSet<FGraphVertexHandle> OtherVertices = OtherIsland->Vertices;
+	for (const FGraphVertexHandle& Node : OtherVertices)
 	{
+		OtherIsland->RemoveVertex(Node);
 		AddVertex(Node);
 	}
 
