merging islands should only trigger the handle on vertex removed event on the other island

diff --git a/Engine/Plugins/Experimental/GameplayGraph/Source/GameplayGraph/Private/Graph/GraphIsland.cpp b/Engine/Plugins/Experimental/GameplayGraph/Source/GameplayGraph/Private/Graph/GraphIsland.cpp
--- a/Engine/Plugins/Experimental/GameplayGraph/Source/GameplayGraph/Private/Graph/GraphIsland.cpp
+++ b/Engine/Plugins/Experimental/GameplayGraph/Source/GameplayGraph/Private/Graph/GraphIsland.cpp
@@ -39,10 +39,9 @@ void UGraphIsland::MergeWith(TObjectPtr<UGraphIsland> OtherIsland)
 		return;
 	}
 
-	TSet<FGraphVertexHandle> OtherVertices = OtherIsland->Vertices;
-	for (const FGraphVertexHandle& Node : OtherVertices)
+	for (const FGraphVertexHandle& Node : OtherIsland->Vertices)
 	{
-		OtherIsland->RemoveVertex(Node);
+		OtherIsland->HandleOnVertexRemoved(Node);
 		AddVertex(Node);
 	}
 
