additional trace scopes for the gameplay graph

diff --git a/Engine/Plugins/Experimental/GameplayGraph/Source/GameplayGraph/Private/Graph/Graph.cpp b/Engine/Plugins/Experimental/GameplayGraph/Source/GameplayGraph/Private/Graph/Graph.cpp
--- a/Engine/Plugins/Experimental/GameplayGraph/Source/GameplayGraph/Private/Graph/Graph.cpp
+++ b/Engine/Plugins/Experimental/GameplayGraph/Source/GameplayGraph/Private/Graph/Graph.cpp
@@ -109,6 +109,7 @@ void UGraph::InitializeFromProperties(const FGraphProperties& InProperties)
 
 FGraphVertexHandle UGraph::CreateVertex(int64 InUniqueIndex)
 {
+	TRACE_CPUPROFILER_EVENT_SCOPE(UGraph::CreateVertex);
 	TObjectPtr<UGraphVertex> Vertex = CreateTypedVertex();
 	if (!Vertex)
 	{
@@ -125,6 +126,7 @@ FGraphVertexHandle UGraph::CreateVertex(int64 InUniqueIndex)
 
 void UGraph::RegisterVertex(TObjectPtr<UGraphVertex> Vertex)
 {
+	TRACE_CPUPROFILER_EVENT_SCOPE(UGraph::RegisterVertex);
 	if (!Vertex)
 	{
 		return;
@@ -138,6 +140,7 @@ void UGraph::RegisterVertex(TObjectPtr<UGraphVertex> Vertex)
 
 FGraphEdgeHandle UGraph::CreateEdge(FGraphVertexHandle Node1, FGraphVertexHandle Node2, int64 InUniqueIndex, bool bAddToIslands)
 {
+	TRACE_CPUPROFILER_EVENT_SCOPE(UGraph::CreateEdge);
 	if (!Node1.IsComplete() || !Node2.IsComplete())
 	{
 		return {};
@@ -191,6 +194,7 @@ FGraphEdgeHandle UGraph::CreateEdge(FGraphVertexHandle Node1, FGraphVertexHandle
 
 void UGraph::RegisterEdge(TObjectPtr<UGraphEdge> Edge)
 {
+	TRACE_CPUPROFILER_EVENT_SCOPE(UGraph::RegisterEdge);
 	if (!Edge)
 	{
 		return;
@@ -208,6 +212,7 @@ void UGraph::RegisterEdge(TObjectPtr<UGraphEdge> Edge)
 
 FGraphIslandHandle UGraph::CreateIsland(TArray<FGraphVertexHandle> InputNodes, int64 InUniqueIndex)
 {
+	TRACE_CPUPROFILER_EVENT_SCOPE(UGraph::CreateIsland);
 	if (InputNodes.IsEmpty())
 	{
 		return {};
@@ -234,6 +239,7 @@ FGraphIslandHandle UGraph::CreateIsland(TArray<FGraphVertexHandle> InputNodes, i
 
 void UGraph::RegisterIsland(TObjectPtr<UGraphIsland> Island)
 {
+	TRACE_CPUPROFILER_EVENT_SCOPE(UGraph::RegisterIsland);
 	if (!Island)
 	{
 		return;
@@ -246,6 +252,7 @@ void UGraph::RegisterIsland(TObjectPtr<UGraphIsland> Island)
 
 void UGraph::RemoveIsland(const FGraphIslandHandle& IslandHandle)
 {
+	TRACE_CPUPROFILER_EVENT_SCOPE(UGraph::RemoveIsland);
 	if (!IslandHandle.IsValid())
 	{
 		return;
@@ -260,6 +267,7 @@ void UGraph::RemoveIsland(const FGraphIslandHandle& IslandHandle)
 
 void UGraph::MergeOrCreateIslands(TObjectPtr<UGraphEdge> Edge)
 {
+	TRACE_CPUPROFILER_EVENT_SCOPE(UGraph::MergeOrCreateIslands);
 	if (!Edge)
 	{
 		return;
@@ -315,6 +323,7 @@ void UGraph::MergeOrCreateIslands(TObjectPtr<UGraphEdge> Edge)
 
 void UGraph::RemoveVertex(const FGraphVertexHandle& NodeHandle)
 {
+	TRACE_CPUPROFILER_EVENT_SCOPE(UGraph::RemoveVertex);
 	if (!NodeHandle.IsValid())
 	{
 		return;
@@ -353,6 +362,7 @@ void UGraph::RemoveVertex(const FGraphVertexHandle& NodeHandle)
 
 void UGraph::RemoveEdge(const FGraphEdgeHandle& EdgeHandle, bool bHandleIslands)
 {
+	TRACE_CPUPROFILER_EVENT_SCOPE(UGraph::RemoveEdge);
 	if (!EdgeHandle.IsValid())
 	{
 		return;
@@ -400,6 +410,7 @@ void UGraph::RemoveEdge(const FGraphEdgeHandle& EdgeHandle, bool bHandleIslands)
 
 void UGraph::RemoveOrSplitIsland(TObjectPtr<UGraphIsland> Island)
 {
+	TRACE_CPUPROFILER_EVENT_SCOPE(UGraph::RemoveOrSplitIsland);
 	if (!Island)
 	{
 		return;
@@ -492,6 +503,7 @@ TObjectPtr<UGraphIsland> UGraph::CreateTypedIsland() const
 
 void UGraph::FinalizeVertex(const FGraphVertexHandle& InHandle)
 {
+	TRACE_CPUPROFILER_EVENT_SCOPE(UGraph::FinalizeVertex);
 	if (!InHandle.IsComplete())
 	{
 		return;
