event for when a graph vertex has its parent island set

diff --git a/Engine/Plugins/Experimental/GameplayGraph/Source/GameplayGraph/Private/Graph/GraphVertex.cpp b/Engine/Plugins/Experimental/GameplayGraph/Source/GameplayGraph/Private/Graph/GraphVertex.cpp
--- a/Engine/Plugins/Experimental/GameplayGraph/Source/GameplayGraph/Private/Graph/GraphVertex.cpp
+++ b/Engine/Plugins/Experimental/GameplayGraph/Source/GameplayGraph/Private/Graph/GraphVertex.cpp
@@ -39,4 +39,10 @@ void UGraphVertex::RemoveEdge(const FGraphEdgeHandle& EdgeHandle)
 void UGraphVertex::HandleOnVertexRemoved()
 {
 	OnVertexRemoved.Broadcast();
+}
+
+void UGraphVertex::SetParentIsland(const FGraphIslandHandle& Island)
+{
+	ParentIsland = Island;
+	OnParentIslandSet.Broadcast(Island);
 }
\ No newline at end of file
diff --git a/Engine/Plugins/Experimental/GameplayGraph/Source/GameplayGraph/Public/Graph/GraphVertex.h b/Engine/Plugins/Experimental/GameplayGraph/Source/GameplayGraph/Public/Graph/GraphVertex.h
--- a/Engine/Plugins/Experimental/GameplayGraph/Source/GameplayGraph/Public/Graph/GraphVertex.h
+++ b/Engine/Plugins/Experimental/GameplayGraph/Source/GameplayGraph/Public/Graph/GraphVertex.h
@@ -7,6 +7,7 @@
 
 /** Event for when the node has been removed from the graph. */
 DECLARE_MULTICAST_DELEGATE(FOnGraphVertexRemoved);
+DECLARE_MULTICAST_DELEGATE_OneParam(FOnGraphVertexParentIslandSet, const FGraphIslandHandle&);
 
 UCLASS()
 class GAMEPLAYGRAPH_API UGraphVertex : public UGraphElement
@@ -33,6 +34,7 @@ public:
 	}
 
 	FOnGraphVertexRemoved OnVertexRemoved;
+	FOnGraphVertexParentIslandSet OnParentIslandSet;
 
 	friend class UGraph;
 	friend class UGraphIsland;
@@ -40,7 +42,7 @@ protected:
 
 	void AddEdgeTo(const FGraphVertexHandle& Node, const FGraphEdgeHandle& Edge);
 	void RemoveEdge(const FGraphEdgeHandle& EdgeHandle);
-	void SetParentIsland(const FGraphIslandHandle& Island) { ParentIsland = Island; }
+	void SetParentIsland(const FGraphIslandHandle& Island);
 
 	virtual void HandleOnVertexRemoved();
 private:
