graph islands get marked as in the process of being destroyed as the start of Destroy

diff --git a/Engine/Plugins/Experimental/GameplayGraph/Source/GameplayGraph/Private/Graph/GraphIsland.cpp b/Engine/Plugins/Experimental/GameplayGraph/Source/GameplayGraph/Private/Graph/GraphIsland.cpp
--- a/Engine/Plugins/Experimental/GameplayGraph/Source/GameplayGraph/Private/Graph/GraphIsland.cpp
+++ b/Engine/Plugins/Experimental/GameplayGraph/Source/GameplayGraph/Private/Graph/GraphIsland.cpp
@@ -9,10 +9,12 @@
 UGraphIsland::UGraphIsland()
 	: Super(EGraphElementType::Island)
 {
+	bPendingDestroy = false;
 }
 
 void UGraphIsland::Destroy()
 {
+	bPendingDestroy = true;
 	TSet<FGraphVertexHandle> VertexCopy = Vertices;
 	Vertices.Empty();
 
diff --git a/Engine/Plugins/Experimental/GameplayGraph/Source/GameplayGraph/Public/Graph/GraphIsland.h b/Engine/Plugins/Experimental/GameplayGraph/Source/GameplayGraph/Public/Graph/GraphIsland.h
--- a/Engine/Plugins/Experimental/GameplayGraph/Source/GameplayGraph/Public/Graph/GraphIsland.h
+++ b/Engine/Plugins/Experimental/GameplayGraph/Source/GameplayGraph/Public/Graph/GraphIsland.h
@@ -41,6 +41,7 @@ public:
 		return FGraphIslandHandle{ GetUniqueIndex(), const_cast<UGraphIsland*>(this) };
 	}
 
+	bool IsPendingDestroy() const { return bPendingDestroy; }
 	bool IsEmpty() const { return Vertices.IsEmpty(); }
 	const TSet<FGraphVertexHandle>& GetVertices() const { return Vertices; }
 	FSerializedIslandData GetSerializedData() const;
@@ -73,4 +74,7 @@ protected:
 private:
 	UPROPERTY(SaveGame)
 	TSet<FGraphVertexHandle> Vertices;
+
+	UPROPERTY(Transient)
+	bool bPendingDestroy;
 };
\ No newline at end of file
