change the gameplay graph plugin to be slightly more useful

diff --git a/Engine/Plugins/Experimental/GameplayGraph/GameplayGraph.uplugin b/Engine/Plugins/Experimental/GameplayGraph/GameplayGraph.uplugin
--- a/Engine/Plugins/Experimental/GameplayGraph/GameplayGraph.uplugin
+++ b/Engine/Plugins/Experimental/GameplayGraph/GameplayGraph.uplugin
@@ -3,7 +3,7 @@
 	"Version" : 1,
 	"VersionName" : "0.1",
 	"FriendlyName" : "Gameplay Graph",
-	"Description" : "Graphs are fun.",
+	"Description" : "A graph representation model and common graph alogrithms that can be used for gameplay.",
 	"Category": "Gameplay",
 	"CreatedBy" : "Epic Games, Inc.",
 	"CreatedByURL" : "http://epicgames.com",
