edges need to be removed from the map of edges that lives on the graph

diff --git a/Engine/Plugins/Experimental/GameplayGraph/Source/GameplayGraph/Private/Graph/Graph.cpp b/Engine/Plugins/Experimental/GameplayGraph/Source/GameplayGraph/Private/Graph/Graph.cpp
--- a/Engine/Plugins/Experimental/GameplayGraph/Source/GameplayGraph/Private/Graph/Graph.cpp
+++ b/Engine/Plugins/Experimental/GameplayGraph/Source/GameplayGraph/Private/Graph/Graph.cpp
@@ -369,6 +369,8 @@ void UGraph::RemoveEdge(const FGraphEdgeHandle& EdgeHandle, bool bHandleIslands)
 	}
 
 	TObjectPtr<UGraphEdge> Edge = EdgeHandle.GetEdge();
+	Edges.Remove(EdgeHandle);
+
 	if (!Edge)
 	{
 		return;
