gameplay graph logging to try and track down an OOM

- fix a potential issue where an incomplete node in an island called for FindConnectedComponents would cause an infinite loop

diff --git a/Engine/Plugins/Experimental/GameplayGraph/Source/GameplayGraph/Private/Graph/Algorithms/Connectivity/ConnectedComponents.cpp b/Engine/Plugins/Experimental/GameplayGraph/Source/GameplayGraph/Private/Graph/Algorithms/Connectivity/ConnectedComponents.cpp
--- a/Engine/Plugins/Experimental/GameplayGraph/Source/GameplayGraph/Private/Graph/Algorithms/Connectivity/ConnectedComponents.cpp
+++ b/Engine/Plugins/Experimental/GameplayGraph/Source/GameplayGraph/Private/Graph/Algorithms/Connectivity/ConnectedComponents.cpp
@@ -8,6 +8,7 @@ namespace Graph::Algorithms
 
 	TArray<TSet<FGraphVertexHandle>> FindConnectedComponents(const TSet<FGraphVertexHandle>& StartingSet)
 	{
+		UE_LOG(LogGameplayGraph, Verbose, TEXT("FindConnectedComponents Starting Set Size: %d"), StartingSet.Num());
 		TArray<TSet<FGraphVertexHandle>> Output;
 		// Do a BFS from each node already assigned to an island.
 		// Every node that we hit that's also in the StartingSet gets
@@ -32,6 +33,13 @@ namespace Graph::Algorithms
 				}
 			);
 
+			if (CurrentComponent.IsEmpty())
+			{
+				// Fallback to prevent adding infinitely many empty components into the array.
+				UE_LOG(LogGameplayGraph, Warning, TEXT("FindConnectedComponents found an empty component with a non-empty work queue [probably due to an incomplete handle?]"));
+				break;
+			}
+
 			Output.Add(CurrentComponent);
 		}
 
diff --git a/Engine/Plugins/Experimental/GameplayGraph/Source/GameplayGraph/Private/Graph/Algorithms/Search/Search.cpp b/Engine/Plugins/Experimental/GameplayGraph/Source/GameplayGraph/Private/Graph/Algorithms/Search/Search.cpp
--- a/Engine/Plugins/Experimental/GameplayGraph/Source/GameplayGraph/Private/Graph/Algorithms/Search/Search.cpp
+++ b/Engine/Plugins/Experimental/GameplayGraph/Source/GameplayGraph/Private/Graph/Algorithms/Search/Search.cpp
@@ -71,14 +71,14 @@ namespace Graph::Algorithms
 			while (!WorkQueue.IsEmpty())
 			{
 				FGraphVertexHandle Next = GetNextAndAdvance<TDataStructure>(WorkQueue);
-				if (!Next.IsComplete())
+				if (Callback(Next))
 				{
-					continue;
+					return Next;
 				}
 
-				if (Callback(Next))
+				if (!Next.IsComplete())
 				{
-					return Next;
+					continue;
 				}
 
 				// Get neighbors and add to the queue.
diff --git a/Engine/Plugins/Experimental/GameplayGraph/Source/GameplayGraph/Private/Graph/GraphHandle.cpp b/Engine/Plugins/Experimental/GameplayGraph/Source/GameplayGraph/Private/Graph/GraphHandle.cpp
--- a/Engine/Plugins/Experimental/GameplayGraph/Source/GameplayGraph/Private/Graph/GraphHandle.cpp
+++ b/Engine/Plugins/Experimental/GameplayGraph/Source/GameplayGraph/Private/Graph/GraphHandle.cpp
@@ -8,6 +8,8 @@
 
 #include UE_INLINE_GENERATED_CPP_BY_NAME(GraphHandle)
 
+DEFINE_LOG_CATEGORY(LogGameplayGraph)
+
 FGraphHandle::FGraphHandle()
 {
 
diff --git a/Engine/Plugins/Experimental/GameplayGraph/Source/GameplayGraph/Public/Graph/GraphHandle.h b/Engine/Plugins/Experimental/GameplayGraph/Source/GameplayGraph/Public/Graph/GraphHandle.h
--- a/Engine/Plugins/Experimental/GameplayGraph/Source/GameplayGraph/Public/Graph/GraphHandle.h
+++ b/Engine/Plugins/Experimental/GameplayGraph/Source/GameplayGraph/Public/Graph/GraphHandle.h
@@ -9,6 +9,8 @@ class UGraphVertex;
 class UGraphEdge;
 class UGraphIsland;
 
+DECLARE_LOG_CATEGORY_EXTERN(LogGameplayGraph, Log, All);
+
 /**
  * For persistence, every node in a graph is given a unique index.
  * A FGraphHandle encapsulates that index to make it easy to go from
