remove smaller gameplay graph island when merging two islands together

diff --git a/Engine/Plugins/Experimental/GameplayGraph/Source/GameplayGraph/Private/Graph/Graph.cpp b/Engine/Plugins/Experimental/GameplayGraph/Source/GameplayGraph/Private/Graph/Graph.cpp
--- a/Engine/Plugins/Experimental/GameplayGraph/Source/GameplayGraph/Private/Graph/Graph.cpp
+++ b/Engine/Plugins/Experimental/GameplayGraph/Source/GameplayGraph/Private/Graph/Graph.cpp
@@ -291,12 +291,17 @@ void UGraph::MergeOrCreateIslands(TObjectPtr<UGraphEdge> Edge)
 
 	if (IslandHandleA.IsValid() && IslandHandleB.IsValid() && IslandHandleA != IslandHandleB)
 	{
-		// Both are in separate, valid islands - merge islands (i.e. use island A and just get rid of island B).
-		if (ensure(IslandA != nullptr))
-		{
-			IslandA->MergeWith(IslandB);
-		}
-		RemoveIsland(IslandHandleB);
+		check(IslandA != nullptr && IslandB != nullptr);
+		// Both are in separate, valid islands - merge islands into a single one. Note that this operation
+		// will remove all vertices from the "other" island. Therefore we want to make sure we keep the larger island.
+		const int32 SizeA = IslandA->Num();
+		const int32 SizeB = IslandB->Num();
+
+		TObjectPtr<UGraphIsland> ToKeepIsland = (SizeA < SizeB) ? IslandB : IslandA;
+		TObjectPtr<UGraphIsland> ToRemoveIsland = (SizeA < SizeB) ? IslandA : IslandB;
+
+		ToKeepIsland->MergeWith(ToRemoveIsland);
+		RemoveIsland(ToRemoveIsland->Handle());
 	}
 	else if (IslandHandleA.IsValid() && !IslandHandleB.IsValid())
 	{
diff --git a/Engine/Plugins/Experimental/GameplayGraph/Source/GameplayGraph/Public/Graph/GraphIsland.h b/Engine/Plugins/Experimental/GameplayGraph/Source/GameplayGraph/Public/Graph/GraphIsland.h
--- a/Engine/Plugins/Experimental/GameplayGraph/Source/GameplayGraph/Public/Graph/GraphIsland.h
+++ b/Engine/Plugins/Experimental/GameplayGraph/Source/GameplayGraph/Public/Graph/GraphIsland.h
@@ -44,6 +44,7 @@ public:
 	bool IsPendingDestroy() const { return bPendingDestroy; }
 	bool IsEmpty() const { return Vertices.IsEmpty(); }
 	const TSet<FGraphVertexHandle>& GetVertices() const { return Vertices; }
+	int32 Num() const { return Vertices.Num(); }
 	FSerializedIslandData GetSerializedData() const;
 
 	FOnGraphIslandVertexAdded OnVertexAdded;
