UGraphVertex::ForEachAdjacentVertex should have the option to be able to also pass up the edge handle being iterated

diff --git a/Engine/Plugins/Experimental/GameplayGraph/Source/GameplayGraph/Public/Graph/GraphVertex.h b/Engine/Plugins/Experimental/GameplayGraph/Source/GameplayGraph/Public/Graph/GraphVertex.h
--- a/Engine/Plugins/Experimental/GameplayGraph/Source/GameplayGraph/Public/Graph/GraphVertex.h
+++ b/Engine/Plugins/Experimental/GameplayGraph/Source/GameplayGraph/Public/Graph/GraphVertex.h
@@ -29,7 +29,14 @@ public:
 	{
 		for (const TPair<FGraphVertexHandle, FGraphEdgeHandle>& Kvp : Edges)
 		{
-			Lambda(Kvp.Key);
+			if constexpr (std::is_invocable_v<TLambda, const FGraphVertexHandle&, const FGraphEdgeHandle&>)
+			{
+				Lambda(Kvp.Key, Kvp.Value);
+			}
+			else
+			{
+				Lambda(Kvp.Key);
+			}
 		}
 	}
 
