add bulk operations of adding edges into the graph and expose ForceRebuildGTParticleGeometry on cluster unions publically to allow AddComponentToCluster to optionally not call it immediately so the call can be bulked

diff --git a/Engine/Plugins/Experimental/GameplayGraph/Source/GameplayGraph/Private/Graph/Graph.cpp b/Engine/Plugins/Experimental/GameplayGraph/Source/GameplayGraph/Private/Graph/Graph.cpp
--- a/Engine/Plugins/Experimental/GameplayGraph/Source/GameplayGraph/Private/Graph/Graph.cpp
+++ b/Engine/Plugins/Experimental/GameplayGraph/Source/GameplayGraph/Private/Graph/Graph.cpp
@@ -186,12 +186,31 @@ FGraphEdgeHandle UGraph::CreateEdge(FGraphVertexHandle Node1, FGraphVertexHandle
 	// If we want to keep track of islands, this is where we need to create/merge islands.
 	if (Properties.bGenerateIslands && bAddToIslands)
 	{
-		MergeOrCreateIslands(Edge);
+		MergeOrCreateIslands({ Edge->Handle() });
 	}
 
 	return Edge->Handle();
 }
 
+void UGraph::CreateBulkEdges(TArray<TPair<FGraphVertexHandle, FGraphVertexHandle>>&& NodesToConnect)
+{
+	// Create all edges normally but don't call MergeOrCreateIslands yet. We'll use the bulk function instead.
+	TArray<FGraphEdgeHandle> NewEdges;
+	NewEdges.Reserve(NodesToConnect.Num());
+
+	for (const TPair<FGraphVertexHandle, FGraphVertexHandle>& NodePair : NodesToConnect)
+	{
+		if (FGraphEdgeHandle Edge = CreateEdge(NodePair.Key, NodePair.Value, -1, false); Edge.IsValid())
+		{
+			NewEdges.Add(Edge);
+		}
+	}
+
+	// MergeOrCreateIslands incrementally determines island connectivity one edge at a time. This is efficient if we're handling a single edge but less
+	// efficient if we're trying to add a bunch of edges all at the same time since it'll cause a vertex to jump between islands. The Bulk function helps prevent that.
+	MergeOrCreateIslands(MoveTemp(NewEdges));
+}
+
 void UGraph::RegisterEdge(TObjectPtr<UGraphEdge> Edge)
 {
 	TRACE_CPUPROFILER_EVENT_SCOPE(UGraph::RegisterEdge);
@@ -265,64 +284,217 @@ void UGraph::RemoveIsland(const FGraphIslandHandle& IslandHandle)
 	Islands.Remove(IslandHandle);
 }
 
-void UGraph::MergeOrCreateIslands(TObjectPtr<UGraphEdge> Edge)
+void UGraph::MergeOrCreateIslands(TArray<FGraphEdgeHandle>&& InEdges)
 {
 	TRACE_CPUPROFILER_EVENT_SCOPE(UGraph::MergeOrCreateIslands);
-	if (!Edge)
+	if (Edges.IsEmpty())
 	{
 		return;
 	}
 
-	FGraphVertexHandle AHandle = Edge->NodeA();
-	FGraphVertexHandle BHandle = Edge->NodeB();
+	// Every vertex needs to track which new island it's going to be in. This can either be an existing island or a completely new island.
+	// This function guarantees that any vertex gets *actually* added into an island exactly once. We're guaranteed that after this function,
+	// every affected vertex is in an island. So the FGraphIslandHandle either points to an existing island (in which case the vertex should be
+	// added to that island) or an invalid island (in which case the vertex is going to be added into a new island). Note that we can assume that
+	// the island IDs will never be negative (and -1 is reserved for "invalid") so we can represent temporary islands as having IDs of -2, -3, etc.
+	TMap<FGraphVertexHandle, FGraphIslandHandle> VertexIslandChanges;
+	VertexIslandChanges.Reserve(InEdges.Num());
+
+	TMap<FGraphIslandHandle, TSet<FGraphVertexHandle>> IslandVertexAdditions;
+	TMap<FGraphIslandHandle, int32> IslandSizeOverride;
 
-	TObjectPtr<UGraphVertex> NodeA = AHandle.GetVertex();
-	TObjectPtr<UGraphVertex> NodeB = BHandle.GetVertex();
-	if (!NodeA || !NodeB)
+	int64 NextTemporaryIslandId = -2;
+
+	auto IsIslandTemporary = [](const FGraphIslandHandle& Handle)
 	{
-		return;
-	}
+		return Handle.GetUniqueIndex() <= -2;
+	};
+
+	auto GetIslandSize = [&IslandVertexAdditions, &IsIslandTemporary, &IslandSizeOverride](const FGraphIslandHandle& Handle)
+	{
+		if (IsIslandTemporary(Handle))
+		{
+			return IslandVertexAdditions.FindRef(Handle).Num();
+		}
+		else if (int32* Size = IslandSizeOverride.Find(Handle))
+		{
+			return *Size;
+		}
+		else if (TObjectPtr<UGraphIsland> Island = Handle.GetIsland())
+		{
+			return Island->Num();
+		}
+		return 0;
+	};
 
-	FGraphIslandHandle IslandHandleA = NodeA->GetParentIsland();
-	FGraphIslandHandle IslandHandleB = NodeB->GetParentIsland();
+	auto RecacheIslandSize = [&IsIslandTemporary, &IslandVertexAdditions, &IslandSizeOverride](const FGraphIslandHandle& Handle, int32 Delta)
+	{
+		if (IsIslandTemporary(Handle))
+		{
+			return;
+		}
 
-	TObjectPtr<UGraphIsland> IslandA = IslandHandleA.GetIsland();
-	TObjectPtr<UGraphIsland> IslandB = IslandHandleB.GetIsland();
+		if (!IslandSizeOverride.Contains(Handle))
+		{
+			int32 Size = 0;
+			if (TObjectPtr<UGraphIsland> Island = Handle.GetIsland())
+			{
+				Size += Island->Num();
+			}
+
+			IslandSizeOverride.Add(Handle, Size);
+		}
 
-	if (IslandHandleA.IsValid() && IslandHandleB.IsValid() && IslandHandleA != IslandHandleB)
+		IslandSizeOverride.Add(Handle, IslandSizeOverride.FindRef(Handle) + Delta);
+	};
+
+	auto ForEveryVertexInIsland = [&IslandVertexAdditions]<typename TLambda>(const FGraphIslandHandle& IslandHandle, TLambda&& Func)
 	{
-		check(IslandA != nullptr && IslandB != nullptr);
-		// Both are in separate, valid islands - merge islands into a single one. Note that this operation
-		// will remove all vertices from the "other" island. Therefore we want to make sure we keep the larger island.
-		const int32 SizeA = IslandA->Num();
-		const int32 SizeB = IslandB->Num();
+		// We need to always iterate over all the vertices in IslandVertexAdditions.
+		for (const FGraphVertexHandle& VertexHandle : IslandVertexAdditions.FindRef(IslandHandle))
+		{
+			Func(VertexHandle);
+		}
+
+		// Is this is a real non-temporary island, we need to iterate over its vertices too. Note that we don't need to
+		// keep track of island vertex removals since we never do partial removal from an island. If an island is being removed,
+		// all its vertices are being moved into a different island.
+		if (TObjectPtr<UGraphIsland> Island = IslandHandle.GetIsland())
+		{
+			for (const FGraphVertexHandle& VertexHandle : Island->GetVertices())
+			{
+				Func(VertexHandle);
+			}
+		}
+	};
+
+	for (const FGraphEdgeHandle& EdgeHandle : InEdges)
+	{
+		if (!EdgeHandle.IsComplete())
+		{
+			continue;
+		}
+
+		TObjectPtr<UGraphEdge> Edge = EdgeHandle.GetEdge();
+		const FGraphVertexHandle& AHandle = Edge->NodeA();
+		const FGraphVertexHandle& BHandle = Edge->NodeB();
 
-		TObjectPtr<UGraphIsland> ToKeepIsland = (SizeA < SizeB) ? IslandB : IslandA;
-		TObjectPtr<UGraphIsland> ToRemoveIsland = (SizeA < SizeB) ? IslandA : IslandB;
+		TObjectPtr<UGraphVertex> NodeA = AHandle.GetVertex();
+		TObjectPtr<UGraphVertex> NodeB = BHandle.GetVertex();
+		if (!NodeA || !NodeB)
+		{
+			return;
+		}
 
-		ToKeepIsland->MergeWith(ToRemoveIsland);
-		RemoveIsland(ToRemoveIsland->Handle());
+		FGraphIslandHandle IslandHandleA = VertexIslandChanges.FindRef(AHandle);
+		if (!IslandHandleA.IsValid())
+		{
+			IslandHandleA = NodeA->GetParentIsland();
+		}
+
+		FGraphIslandHandle IslandHandleB = VertexIslandChanges.FindRef(BHandle);
+		if (!IslandHandleB.IsValid())
+		{
+			IslandHandleB = NodeB->GetParentIsland();
+		}
+
+		if (IslandHandleA.IsValid() && IslandHandleB.IsValid() && IslandHandleA != IslandHandleB)
+		{
+			// We need to move all the vertices in one island to the other.
+			FGraphIslandHandle ToKeepIsland;
+			FGraphIslandHandle ToRemoveIsland;
+
+			const bool bIsATemporary = IsIslandTemporary(IslandHandleA);
+			const bool bIsBTemporary = IsIslandTemporary(IslandHandleB);
+			if (bIsATemporary == bIsBTemporary)
+			{
+				// In the case both islands are not temporary (or both are temporary), choose the larger island to add to.
+				const int32 SizeA = GetIslandSize(IslandHandleA);
+				const int32 SizeB = GetIslandSize(IslandHandleB);
+				ToKeepIsland = (SizeA > SizeB) ? IslandHandleA : IslandHandleB;
+				ToRemoveIsland = (SizeA > SizeB) ? IslandHandleB : IslandHandleA;
+			}
+			else
+			{
+				// In the case that only one of the islands is temporary, choose the non-temporary island.
+				ToKeepIsland = bIsATemporary ? IslandHandleB : IslandHandleA;
+				ToRemoveIsland = bIsATemporary ? IslandHandleA : IslandHandleB;
+			}
+
+			int32 VerticesChanged = 0;
+			ForEveryVertexInIsland(ToRemoveIsland,
+				[&VertexIslandChanges, &IslandVertexAdditions, &ToKeepIsland, &VerticesChanged](const FGraphVertexHandle& VertexHandle)
+				{
+					VertexIslandChanges.Add(VertexHandle, ToKeepIsland);
+					IslandVertexAdditions.FindOrAdd(ToKeepIsland).Add(VertexHandle);
+					++VerticesChanged;
+				}
+			);
+
+			IslandVertexAdditions.Remove(ToRemoveIsland);
+			RecacheIslandSize(ToKeepIsland, VerticesChanged);
+			RecacheIslandSize(ToRemoveIsland, -VerticesChanged);
+		}
+		else if (IslandHandleA.IsValid() && !IslandHandleB.IsValid())
+		{
+			VertexIslandChanges.Add(BHandle, IslandHandleA);
+			IslandVertexAdditions.FindOrAdd(IslandHandleA).Add(BHandle);
+			RecacheIslandSize(IslandHandleA, 1);
+		}
+		else if (IslandHandleB.IsValid() && !IslandHandleA.IsValid())
+		{
+			VertexIslandChanges.Add(AHandle, IslandHandleB);
+			IslandVertexAdditions.FindOrAdd(IslandHandleB).Add(AHandle);
+			RecacheIslandSize(IslandHandleB, 1);
+		}
+		else if (!IslandHandleA.IsValid() && !IslandHandleB.IsValid())
+		{
+			// Neither is in an island - need to create a new one.
+			FGraphIslandHandle NewIsland{ NextTemporaryIslandId--, nullptr };
+			VertexIslandChanges.Add(AHandle, NewIsland);
+			VertexIslandChanges.Add(BHandle, NewIsland);
+			IslandVertexAdditions.Emplace(NewIsland, TSet<FGraphVertexHandle>{AHandle, BHandle});
+		}
 	}
-	else if (IslandHandleA.IsValid() && !IslandHandleB.IsValid())
+
+	// Now that we have all the changes we want to make, we can start making them! Iterate over
+	// VertexIslandChanges and handle additions into an existing island first.
+	for (const TPair<FGraphVertexHandle, FGraphIslandHandle>& Change : VertexIslandChanges)
 	{
-		// Only A is in an island, go into island A.
-		if (ensure(IslandA != nullptr))
+		if (IsIslandTemporary(Change.Value))
+		{
+			continue;
+		}
+
+		if (TObjectPtr<UGraphIsland> Island = Change.Value.GetIsland())
 		{
-			IslandA->AddVertex(BHandle);
+			Island->AddVertex(Change.Key);
 		}
 	}
-	else if (IslandHandleB.IsValid() && !IslandHandleA.IsValid())
+
+	// Next handle the creation of any temporary islands.
+	for (const TPair<FGraphIslandHandle, TSet<FGraphVertexHandle>>& Change : IslandVertexAdditions)
 	{
-		// Only B is in an island, go into island B.
-		if (ensure(IslandB != nullptr))
+		if (!IsIslandTemporary(Change.Key))
 		{
-			IslandB->AddVertex(AHandle);
+			continue;
 		}
+
+		CreateIsland(Change.Value.Array());
 	}
-	else if (!IslandHandleA.IsValid() && !IslandHandleB.IsValid())
+
+	// Destroy all empty islands
+	for (const TPair<FGraphIslandHandle, int32>& Change : IslandSizeOverride)
 	{
-		// Neither is in an island - need to create a new one.
-		CreateIsland({ AHandle, BHandle });
+		if (IsIslandTemporary(Change.Key))
+		{
+			continue;
+		}
+
+		if (Change.Value == 0)
+		{
+			RemoveIsland(Change.Key);
+		}
 	}
 }
 
diff --git a/Engine/Plugins/Experimental/GameplayGraph/Source/GameplayGraph/Private/Graph/GraphIsland.cpp b/Engine/Plugins/Experimental/GameplayGraph/Source/GameplayGraph/Private/Graph/GraphIsland.cpp
--- a/Engine/Plugins/Experimental/GameplayGraph/Source/GameplayGraph/Private/Graph/GraphIsland.cpp
+++ b/Engine/Plugins/Experimental/GameplayGraph/Source/GameplayGraph/Private/Graph/GraphIsland.cpp
@@ -14,6 +14,7 @@ UGraphIsland::UGraphIsland()
 
 void UGraphIsland::Destroy()
 {
+	TRACE_CPUPROFILER_EVENT_SCOPE(UGraphIsland::Destroy);
 	bPendingDestroy = true;
 	TSet<FGraphVertexHandle> VertexCopy = Vertices;
 	Vertices.Empty();
@@ -34,24 +35,9 @@ void UGraphIsland::Destroy()
 	HandleOnDestroyed();
 }
 
-void UGraphIsland::MergeWith(TObjectPtr<UGraphIsland> OtherIsland)
-{
-	if (!OtherIsland)
-	{
-		return;
-	}
-
-	for (const FGraphVertexHandle& Node : OtherIsland->Vertices)
-	{
-		OtherIsland->HandleOnVertexRemoved(Node);
-		AddVertex(Node);
-	}
-
-	OtherIsland->Vertices.Empty();
-}
-
 void UGraphIsland::AddVertex(const FGraphVertexHandle& Node)
 {
+	TRACE_CPUPROFILER_EVENT_SCOPE(UGraphIsland::AddVertex);
 	if (!Node.IsValid())
 	{
 		return;
@@ -62,6 +48,15 @@ void UGraphIsland::AddVertex(const FGraphVertexHandle& Node)
 	{
 		return;
 	}
+
+	if (FGraphIslandHandle OldIslandHandle = NodePtr->GetParentIsland(); OldIslandHandle.IsComplete())
+	{
+		if (TObjectPtr<UGraphIsland> OldIsland = OldIslandHandle.GetIsland())
+		{
+			OldIsland->RemoveVertex(Node);
+		}
+	}
+
 	NodePtr->SetParentIsland(Handle());
 	Vertices.Add(Node);
 	HandleOnVertexAdded(Node);
@@ -69,6 +64,7 @@ void UGraphIsland::AddVertex(const FGraphVertexHandle& Node)
 
 void UGraphIsland::RemoveVertex(const FGraphVertexHandle& Node)
 {
+	TRACE_CPUPROFILER_EVENT_SCOPE(UGraphIsland::RemoveVertex);
 	if (!Node.IsValid())
 	{
 		return;
diff --git a/Engine/Plugins/Experimental/GameplayGraph/Source/GameplayGraph/Public/Graph/Graph.h b/Engine/Plugins/Experimental/GameplayGraph/Source/GameplayGraph/Public/Graph/Graph.h
--- a/Engine/Plugins/Experimental/GameplayGraph/Source/GameplayGraph/Public/Graph/Graph.h
+++ b/Engine/Plugins/Experimental/GameplayGraph/Source/GameplayGraph/Public/Graph/Graph.h
@@ -96,6 +96,9 @@ public:
 	/** Creates an edge between the two nodes. */
 	FGraphEdgeHandle CreateEdge(FGraphVertexHandle Node1, FGraphVertexHandle Node2, int64 InUniqueIndex = INDEX_NONE, bool bAddToIslands = true);
 
+	/** Creates edges in bulk. This is more efficient than calling CreateEdge multiple times since we will only try to assign a node to an island once. */
+	void CreateBulkEdges(TArray<TPair<FGraphVertexHandle, FGraphVertexHandle>>&& NodesToConnect);
+
 	/** Removes a node from the graph along with any edges that contain it. */
 	void RemoveVertex(const FGraphVertexHandle& NodeHandle);
 
@@ -158,8 +161,8 @@ private:
 	/** Add an island to the graph and modifies the NextAvailableIslandUniqueIndex as necessary to maintain its validity. */
 	void RegisterIsland(TObjectPtr<UGraphIsland> Edge);
 
-	/** When we add an edge into a graph, we will want to merge islands if the edge connects two nodes that are in different islands. */
-	void MergeOrCreateIslands(TObjectPtr<UGraphEdge> Edge);
+	/** When we add multiple edges into the graph. This function will ensure that the interactions we make externally are kept to a minimum. */
+	void MergeOrCreateIslands(TArray<FGraphEdgeHandle>&& InEdges);
 
 	/** After a change, this function will remove the island if it's empty or will attempt to split it into two smaller islands. */
 	void RemoveOrSplitIsland(TObjectPtr<UGraphIsland> Island);
diff --git a/Engine/Plugins/Experimental/GameplayGraph/Source/GameplayGraph/Public/Graph/GraphIsland.h b/Engine/Plugins/Experimental/GameplayGraph/Source/GameplayGraph/Public/Graph/GraphIsland.h
--- a/Engine/Plugins/Experimental/GameplayGraph/Source/GameplayGraph/Public/Graph/GraphIsland.h
+++ b/Engine/Plugins/Experimental/GameplayGraph/Source/GameplayGraph/Public/Graph/GraphIsland.h
@@ -63,9 +63,6 @@ protected:
 	/** Called when removing the island from the graph. */
 	void Destroy();
 
-	/** Takes the nodes from the other island and puts them into this island. */
-	void MergeWith(TObjectPtr<UGraphIsland> OtherIsland);
-
 	/** Adds a single node into this island. */
 	void AddVertex(const FGraphVertexHandle& Node);
 
diff --git a/Engine/Source/Runtime/Engine/Classes/PhysicsEngine/ClusterUnionComponent.h b/Engine/Source/Runtime/Engine/Classes/PhysicsEngine/ClusterUnionComponent.h
--- a/Engine/Source/Runtime/Engine/Classes/PhysicsEngine/ClusterUnionComponent.h
+++ b/Engine/Source/Runtime/Engine/Classes/PhysicsEngine/ClusterUnionComponent.h
@@ -127,7 +127,7 @@ public:
 	ENGINE_API UClusterUnionComponent(const FObjectInitializer& ObjectInitializer);
 
 	UFUNCTION(BlueprintCallable, Category="Cluster Union")
-	ENGINE_API void AddComponentToCluster(UPrimitiveComponent* InComponent, const TArray<int32>& BoneIds);
+	ENGINE_API void AddComponentToCluster(UPrimitiveComponent* InComponent, const TArray<int32>& BoneIds, bool bRebuildGeometry = true);
 
 	UFUNCTION(BlueprintCallable, Category = "Cluster Union")
 	ENGINE_API void RemoveComponentFromCluster(UPrimitiveComponent* InComponent);
@@ -171,6 +171,10 @@ public:
 
 	ENGINE_API int32 NumChildClusterComponents() const;
 
+	// Force a rebuild of the GT geometry. This needs to happen immediately when we add/remove on the GT so that the SQ is up to date
+	// and doesn't need to wait for the next OnSyncBodies.
+	ENGINE_API void ForceRebuildGTParticleGeometry();
+
 	friend class UClusterUnionReplicatedProxyComponent;
 protected:
 
@@ -254,10 +258,6 @@ private:
 	UFUNCTION()
 	ENGINE_API bool IsAuthority() const;
 
-	// Force a rebuild of the GT geometry. This needs to happen immediately when we add/remove on the GT so that the SQ is up to date
-	// and doesn't need to wait for the next OnSyncBodies.
-	ENGINE_API void ForceRebuildGTParticleGeometry();
-
 	//~ Begin UActorComponent Interface
 public:
 	ENGINE_API virtual void OnCreatePhysicsState() override;
diff --git a/Engine/Source/Runtime/Engine/Private/PhysicsEngine/ClusterUnionComponent.cpp b/Engine/Source/Runtime/Engine/Private/PhysicsEngine/ClusterUnionComponent.cpp
--- a/Engine/Source/Runtime/Engine/Private/PhysicsEngine/ClusterUnionComponent.cpp
+++ b/Engine/Source/Runtime/Engine/Private/PhysicsEngine/ClusterUnionComponent.cpp
@@ -72,7 +72,7 @@ FPhysScene_Chaos* UClusterUnionComponent::GetChaosScene() const
 	return GWorld->GetPhysicsScene();
 }
 
-void UClusterUnionComponent::AddComponentToCluster(UPrimitiveComponent* InComponent, const TArray<int32>& BoneIds)
+void UClusterUnionComponent::AddComponentToCluster(UPrimitiveComponent* InComponent, const TArray<int32>& BoneIds, bool bRebuildGeometry)
 {
 	if (!InComponent || !PhysicsProxy || !InComponent->GetWorld())
 	{
@@ -153,7 +153,11 @@ void UClusterUnionComponent::AddComponentToCluster(UPrimitiveComponent* InCompon
 	PendingComponentSync.Add(InComponent, PendingData);
 
 	PhysicsProxy->AddPhysicsObjects_External(Objects);
-	ForceRebuildGTParticleGeometry();
+
+	if (bRebuildGeometry)
+	{
+		ForceRebuildGTParticleGeometry();
+	}
 }
 
 void UClusterUnionComponent::RemoveComponentFromCluster(UPrimitiveComponent* InComponent)
