add ability to bulk remove nodes from the gameplay graph

diff --git a/Engine/Plugins/Experimental/GameplayGraph/Source/GameplayGraph/Private/Graph/Graph.cpp b/Engine/Plugins/Experimental/GameplayGraph/Source/GameplayGraph/Private/Graph/Graph.cpp
--- a/Engine/Plugins/Experimental/GameplayGraph/Source/GameplayGraph/Private/Graph/Graph.cpp
+++ b/Engine/Plugins/Experimental/GameplayGraph/Source/GameplayGraph/Private/Graph/Graph.cpp
@@ -500,41 +500,59 @@ void UGraph::MergeOrCreateIslands(TArray<FGraphEdgeHandle>&& InEdges)
 
 void UGraph::RemoveVertex(const FGraphVertexHandle& NodeHandle)
 {
-	TRACE_CPUPROFILER_EVENT_SCOPE(UGraph::RemoveVertex);
-	if (!NodeHandle.IsValid())
-	{
-		return;
-	}
+	RemoveBulkVertices({ NodeHandle });
+}
 
-	TObjectPtr<UGraphVertex> Node = NodeHandle.GetVertex();
-	if (!Node)
+void UGraph::RemoveBulkVertices(const TArray<FGraphVertexHandle>& InHandles)
+{
+	TRACE_CPUPROFILER_EVENT_SCOPE(UGraph::RemoveBulkVertices);
+
+	TSet<FGraphIslandHandle> AffectedIslands;
+
+	for (const FGraphVertexHandle& NodeHandle : InHandles)
 	{
-		return;
-	}
+		if (NodeHandle.IsValid())
+		{
+			// We must remove every edge this node is a part of.
+			for (const FGraphEdgeHandle& EdgeHandle : VertexEdges.FindOrAdd(NodeHandle))
+			{
+				// Don't immediately handle islands. We'll do it later.
+				RemoveEdge(EdgeHandle, false);
+			}
 
-	// We must remove every edge this node is a part of.
-	for (const FGraphEdgeHandle& EdgeHandle : VertexEdges.FindOrAdd(NodeHandle))
+			if (TObjectPtr<UGraphVertex> Node = NodeHandle.GetVertex())
+			{
+				if (TObjectPtr<UGraphIsland> Island = Node->GetParentIsland().GetIsland())
+				{
+					AffectedIslands.Add(Node->GetParentIsland());
+					Island->RemoveVertex(NodeHandle);
+				}
+			}
+		}
+	}
+	
+	if (Properties.bGenerateIslands)
 	{
-		// Don't immediately handle islands. We'll do it later.
-		RemoveEdge(EdgeHandle, false);
+		for (const FGraphIslandHandle& IslandHandle : AffectedIslands)
+		{
+			RemoveOrSplitIsland(IslandHandle.GetIsland());
+		}
 	}
 
-	if (Properties.bGenerateIslands)
+	// A final pass after generation of islands to clean up book-keeping.
+	// TODO: Not sure if this is necessary and can be done before we regenerate islands?
+	for (const FGraphVertexHandle& NodeHandle : InHandles)
 	{
-		// If the node is a part of an island - remove it from the island and evaluate the effect that has.
-		if (const FGraphIslandHandle& IslandHandle = Node->GetParentIsland(); IslandHandle.IsValid())
+		if (NodeHandle.IsValid())
 		{
-			if (TObjectPtr<UGraphIsland> Island = IslandHandle.GetIsland())
+			if (TObjectPtr<UGraphVertex> Node = NodeHandle.GetVertex())
 			{
-				Island->RemoveVertex(NodeHandle);
-				RemoveOrSplitIsland(Island);
+				Node->HandleOnVertexRemoved();
 			}
+			Vertices.Remove(NodeHandle);
+			VertexEdges.Remove(NodeHandle);
 		}
 	}
-
-	Node->HandleOnVertexRemoved();
-	Vertices.Remove(NodeHandle);
-	VertexEdges.Remove(NodeHandle);
 }
 
 void UGraph::RemoveEdge(const FGraphEdgeHandle& EdgeHandle, bool bHandleIslands)
diff --git a/Engine/Plugins/Experimental/GameplayGraph/Source/GameplayGraph/Public/Graph/Graph.h b/Engine/Plugins/Experimental/GameplayGraph/Source/GameplayGraph/Public/Graph/Graph.h
--- a/Engine/Plugins/Experimental/GameplayGraph/Source/GameplayGraph/Public/Graph/Graph.h
+++ b/Engine/Plugins/Experimental/GameplayGraph/Source/GameplayGraph/Public/Graph/Graph.h
@@ -102,6 +102,8 @@ public:
 	/** Removes a node from the graph along with any edges that contain it. */
 	void RemoveVertex(const FGraphVertexHandle& NodeHandle);
 
+	void RemoveBulkVertices(const TArray<FGraphVertexHandle>& InHandles);
+
 	/** Removes an edge from the graph. */
 	void RemoveEdge(const FGraphEdgeHandle& EdgeHandle, bool bHandleIslands);
 
