Gameplay graph : Add method to find edge between two vertices of the graph

diff --git a/Engine/Plugins/Experimental/GameplayGraph/Source/GameplayGraph/Private/Graph/GraphVertex.cpp b/Engine/Plugins/Experimental/GameplayGraph/Source/GameplayGraph/Private/Graph/GraphVertex.cpp
--- a/Engine/Plugins/Experimental/GameplayGraph/Source/GameplayGraph/Private/Graph/GraphVertex.cpp
+++ b/Engine/Plugins/Experimental/GameplayGraph/Source/GameplayGraph/Private/Graph/GraphVertex.cpp
@@ -15,6 +15,11 @@ bool UGraphVertex::HasEdgeTo(const FGraphVertexHandle& Other) const
 	return Edges.Contains(Other);
 }
 
+const FGraphEdgeHandle* UGraphVertex::FindEdgeTo(const FGraphVertexHandle& Other) const
+{
+	return Edges.Find(Other);
+}
+
 void UGraphVertex::AddEdgeTo(const FGraphVertexHandle& Node, const FGraphEdgeHandle& Edge)
 {
 	Edges.Add(Node, Edge);
diff --git a/Engine/Plugins/Experimental/GameplayGraph/Source/GameplayGraph/Public/Graph/GraphVertex.h b/Engine/Plugins/Experimental/GameplayGraph/Source/GameplayGraph/Public/Graph/GraphVertex.h
--- a/Engine/Plugins/Experimental/GameplayGraph/Source/GameplayGraph/Public/Graph/GraphVertex.h
+++ b/Engine/Plugins/Experimental/GameplayGraph/Source/GameplayGraph/Public/Graph/GraphVertex.h
@@ -22,6 +22,8 @@ public:
 	}
 
 	bool HasEdgeTo(const FGraphVertexHandle& Other) const;
+	const FGraphEdgeHandle* FindEdgeTo(const FGraphVertexHandle& Other) const;
+
 	const FGraphIslandHandle& GetParentIsland() const { return ParentIsland; }
 
 	template<typename TLambda>
