Added a method to walk all the vertices in a graph island and call a supplied lambda

#rnx

diff --git a/Engine/Plugins/Experimental/GameplayGraph/Source/GameplayGraph/Public/Graph/GraphIsland.h b/Engine/Plugins/Experimental/GameplayGraph/Source/GameplayGraph/Public/Graph/GraphIsland.h
--- a/Engine/Plugins/Experimental/GameplayGraph/Source/GameplayGraph/Public/Graph/GraphIsland.h
+++ b/Engine/Plugins/Experimental/GameplayGraph/Source/GameplayGraph/Public/Graph/GraphIsland.h
@@ -47,6 +47,15 @@ public:
 	int32 Num() const { return Vertices.Num(); }
 	FSerializedIslandData GetSerializedData() const;
 
+	template<typename TLambda>
+	void ForEachVertex(TLambda&& Lambda)
+	{
+		for (const FGraphVertexHandle& Vh : Vertices)
+		{
+			Lambda(Vh);
+		}
+	}
+
 	FOnGraphIslandVertexAdded OnVertexAdded;
 	FOnGraphIslandVertexRemoved OnVertexRemoved;
 	FOnGraphIslandDestroyed OnDestroyed;
