Chaos : ConnectivityComponent persistence update :
Make UConnectivityComponent:GraphHandle transient to reduce the need for saving ConnectivityComponents
-	Added a FGraphUniqueIndex to wrap how the index is internally stored.
-	Updated vertex and edge creation to rely on the FGraphUniqueIndex saved on the ConnectionManagerComponent, which is reattached during de-persistence operations.

diff --git a/Engine/Plugins/Experimental/GameplayGraph/Source/GameplayGraph/Private/Graph/Graph.cpp b/Engine/Plugins/Experimental/GameplayGraph/Source/GameplayGraph/Private/Graph/Graph.cpp
--- a/Engine/Plugins/Experimental/GameplayGraph/Source/GameplayGraph/Private/Graph/Graph.cpp
+++ b/Engine/Plugins/Experimental/GameplayGraph/Source/GameplayGraph/Private/Graph/Graph.cpp
@@ -96,9 +96,6 @@ void UGraph::Empty()
 	Vertices.Empty();
 	Edges.Empty();
 	Islands.Empty();
-	NextAvailableVertexUniqueIndex = 0;
-	NextAvailableEdgeUniqueIndex = 0;
-	NextAvailableIslandUniqueIndex = 0;
 }
 
 void UGraph::InitializeFromProperties(const FGraphProperties& InProperties)
@@ -107,7 +104,7 @@ void UGraph::InitializeFromProperties(const FGraphProperties& InProperties)
 	Properties = InProperties;
 }
 
-FGraphVertexHandle UGraph::CreateVertex(int64 InUniqueIndex)
+FGraphVertexHandle UGraph::CreateVertex(FGraphUniqueIndex InUniqueIndex)
 {
 	TRACE_CPUPROFILER_EVENT_SCOPE(UGraph::CreateVertex);
 	TObjectPtr<UGraphVertex> Vertex = CreateTypedVertex();
@@ -115,10 +112,13 @@ FGraphVertexHandle UGraph::CreateVertex(int64 InUniqueIndex)
 	{
 		return {};
 	}
-	Vertex->OnCreate();
 
-	// No need to increment since AddNode will do that.
-	Vertex->SetUniqueIndex((InUniqueIndex == INDEX_NONE) ? NextAvailableVertexUniqueIndex : InUniqueIndex);
+	if (!ensure(InUniqueIndex.IsValid()))
+	{
+		return{};
+	}
+	Vertex->OnCreate();
+	Vertex->SetUniqueIndex(InUniqueIndex);
 	RegisterVertex(Vertex);
 	OnVertexCreated.Broadcast(Vertex->Handle());
 	return Vertex->Handle();
@@ -135,10 +135,9 @@ void UGraph::RegisterVertex(TObjectPtr<UGraphVertex> Vertex)
 	const FGraphVertexHandle Handle = Vertex->Handle();
 	Vertex->SetParentGraph(this);
 	Vertices.Add(Handle, Vertex);
-	NextAvailableVertexUniqueIndex = FMath::Max(NextAvailableVertexUniqueIndex, Handle.GetUniqueIndex() + 1);
 }
 
-FGraphEdgeHandle UGraph::CreateEdge(FGraphVertexHandle Node1, FGraphVertexHandle Node2, int64 InUniqueIndex, bool bAddToIslands)
+FGraphEdgeHandle UGraph::CreateEdge(FGraphVertexHandle Node1, FGraphVertexHandle Node2, FGraphUniqueIndex InUniqueIndex, bool bAddToIslands)
 {
 	TRACE_CPUPROFILER_EVENT_SCOPE(UGraph::CreateEdge);
 	if (!Node1.IsComplete() || !Node2.IsComplete())
@@ -169,10 +168,14 @@ FGraphEdgeHandle UGraph::CreateEdge(FGraphVertexHandle Node1, FGraphVertexHandle
 	{
 		return {};
 	}
-	Edge->OnCreate();
 
-	// No need to increment since AddEdge will do that.
-	Edge->SetUniqueIndex((InUniqueIndex == INDEX_NONE) ? NextAvailableEdgeUniqueIndex : InUniqueIndex);
+	if (!ensure(InUniqueIndex.IsValid()))
+	{
+		return{};
+	}
+
+	Edge->OnCreate();
+	Edge->SetUniqueIndex(InUniqueIndex);
 	Edge->SetNodes(Node1, Node2);
 
 	RegisterEdge(Edge);
@@ -200,7 +203,7 @@ void UGraph::CreateBulkEdges(TArray<TPair<FGraphVertexHandle, FGraphVertexHandle
 
 	for (const TPair<FGraphVertexHandle, FGraphVertexHandle>& NodePair : NodesToConnect)
 	{
-		if (FGraphEdgeHandle Edge = CreateEdge(NodePair.Key, NodePair.Value, -1, false); Edge.IsValid())
+		if (FGraphEdgeHandle Edge = CreateEdge(NodePair.Key, NodePair.Value, FGraphUniqueIndex::CreateUniqueIndex(), false); Edge.IsValid())
 		{
 			NewEdges.Add(Edge);
 		}
@@ -225,11 +228,9 @@ void UGraph::RegisterEdge(TObjectPtr<UGraphEdge> Edge)
 
 	VertexEdges.FindOrAdd(Edge->NodeA()).Add(Handle);
 	VertexEdges.FindOrAdd(Edge->NodeB()).Add(Handle);
-
-	NextAvailableEdgeUniqueIndex = FMath::Max(NextAvailableEdgeUniqueIndex, Handle.GetUniqueIndex() + 1);
 }
 
-FGraphIslandHandle UGraph::CreateIsland(TArray<FGraphVertexHandle> InputNodes, int64 InUniqueIndex)
+FGraphIslandHandle UGraph::CreateIsland(TArray<FGraphVertexHandle> InputNodes, FGraphUniqueIndex InUniqueIndex)
 {
 	TRACE_CPUPROFILER_EVENT_SCOPE(UGraph::CreateIsland);
 	if (InputNodes.IsEmpty())
@@ -243,8 +244,7 @@ FGraphIslandHandle UGraph::CreateIsland(TArray<FGraphVertexHandle> InputNodes, i
 		return {};
 	}
 	Island->OnCreate();
-
-	Island->SetUniqueIndex((InUniqueIndex == INDEX_NONE) ? NextAvailableIslandUniqueIndex : InUniqueIndex);
+	Island->SetUniqueIndex((!InUniqueIndex.IsValid()) ? FGraphUniqueIndex::CreateUniqueIndex() : InUniqueIndex);
 
 	OnIslandCreated.Broadcast(Island->Handle());
 	for (const FGraphVertexHandle& Node : InputNodes)
@@ -266,7 +266,6 @@ void UGraph::RegisterIsland(TObjectPtr<UGraphIsland> Island)
 	const FGraphIslandHandle Handle = Island->Handle();
 	Island->SetParentGraph(this);
 	Islands.Add(Handle, Island);
-	NextAvailableIslandUniqueIndex = FMath::Max(NextAvailableIslandUniqueIndex, Handle.GetUniqueIndex() + 1);
 }
 
 void UGraph::RemoveIsland(const FGraphIslandHandle& IslandHandle)
@@ -303,11 +302,12 @@ void UGraph::MergeOrCreateIslands(TArray<FGraphEdgeHandle>&& InEdges)
 	TMap<FGraphIslandHandle, TSet<FGraphVertexHandle>> IslandVertexAdditions;
 	TMap<FGraphIslandHandle, int32> IslandSizeOverride;
 
-	int64 NextTemporaryIslandId = -2;
+	FGraphUniqueIndex NextTemporaryIslandId;
+	NextTemporaryIslandId.SetTemporary(true);
 
 	auto IsIslandTemporary = [](const FGraphIslandHandle& Handle)
 	{
-		return Handle.GetUniqueIndex() <= -2;
+		return Handle.GetUniqueIndex().IsTemporary();
 	};
 
 	auto GetIslandSize = [&IslandVertexAdditions, &IsIslandTemporary, &IslandSizeOverride](const FGraphIslandHandle& Handle)
@@ -450,7 +450,7 @@ void UGraph::MergeOrCreateIslands(TArray<FGraphEdgeHandle>&& InEdges)
 		else if (!IslandHandleA.IsValid() && !IslandHandleB.IsValid())
 		{
 			// Neither is in an island - need to create a new one.
-			FGraphIslandHandle NewIsland{ NextTemporaryIslandId--, nullptr };
+			FGraphIslandHandle NewIsland{ NextTemporaryIslandId.NextUniqueIndex(), nullptr};
 			VertexIslandChanges.Add(AHandle, NewIsland);
 			VertexIslandChanges.Add(BHandle, NewIsland);
 			IslandVertexAdditions.Emplace(NewIsland, TSet<FGraphVertexHandle>{AHandle, BHandle});
@@ -667,14 +667,14 @@ void UGraph::RemoveOrSplitIsland(TObjectPtr<UGraphIsland> Island)
 	}
 }
 
-FGraphVertexHandle UGraph::GetCompleteNodeHandle(const FGraphVertexHandle& InHandle)
+FGraphVertexHandle UGraph::GetCompleteNodeHandle(const FGraphVertexHandle& InHandle) const
 {
 	if (!InHandle.IsValid() || InHandle.HasElement())
 	{
 		return InHandle;
 	}
 
-	TObjectPtr<UGraphVertex>* Data = Vertices.Find(InHandle);
+	const TObjectPtr<UGraphVertex>* Data = Vertices.Find(InHandle);
 	if (!Data || !*Data)
 	{
 		return {};
diff --git a/Engine/Plugins/Experimental/GameplayGraph/Source/GameplayGraph/Private/Graph/GraphHandle.cpp b/Engine/Plugins/Experimental/GameplayGraph/Source/GameplayGraph/Private/Graph/GraphHandle.cpp
--- a/Engine/Plugins/Experimental/GameplayGraph/Source/GameplayGraph/Private/Graph/GraphHandle.cpp
+++ b/Engine/Plugins/Experimental/GameplayGraph/Source/GameplayGraph/Private/Graph/GraphHandle.cpp
@@ -5,6 +5,7 @@
 #include "Graph/GraphEdge.h"
 #include "Graph/GraphIsland.h"
 #include "Graph/GraphVertex.h"
+#include "Misc/Guid.h"
 
 #include UE_INLINE_GENERATED_CPP_BY_NAME(GraphHandle)
 
@@ -15,7 +16,7 @@ FGraphHandle::FGraphHandle()
 
 }
 
-FGraphHandle::FGraphHandle(int64 InUniqueIndex, TObjectPtr<UGraphElement> InElement)
+FGraphHandle::FGraphHandle(FGraphUniqueIndex InUniqueIndex, TObjectPtr<UGraphElement> InElement)
 	: UniqueIndex(InUniqueIndex)
 	, Element(InElement)
 {
@@ -33,7 +34,7 @@ TObjectPtr<UGraphElement> FGraphHandle::GetElement() const
 
 bool FGraphHandle::IsValid() const
 {
-	return UniqueIndex != INDEX_NONE;
+	return UniqueIndex.IsValid();
 }
 
 bool FGraphHandle::HasElement() const
@@ -63,13 +64,13 @@ bool FGraphHandle::operator<(const FGraphHandle& Other) const
 
 void FGraphHandle::Clear()
 {
-	UniqueIndex = INDEX_NONE;
+	UniqueIndex = FGraphUniqueIndex();
 	Element = nullptr;
 }
 
 uint32 GetTypeHash(const FGraphHandle& Handle)
 {
-	return ::GetTypeHash(Handle.UniqueIndex);
+	return GetTypeHash(Handle.UniqueIndex);
 }
 
 FGraphVertexHandle::FGraphVertexHandle()
diff --git a/Engine/Plugins/Experimental/GameplayGraph/Source/GameplayGraph/Public/Graph/Graph.h b/Engine/Plugins/Experimental/GameplayGraph/Source/GameplayGraph/Public/Graph/Graph.h
--- a/Engine/Plugins/Experimental/GameplayGraph/Source/GameplayGraph/Public/Graph/Graph.h
+++ b/Engine/Plugins/Experimental/GameplayGraph/Source/GameplayGraph/Public/Graph/Graph.h
@@ -88,13 +88,13 @@ public:
 	void InitializeFromProperties(const FGraphProperties& Properties);
 
 	/** Given a node handle, find the handle in the current graph with a proper element set. */
-	FGraphVertexHandle GetCompleteNodeHandle(const FGraphVertexHandle& InHandle);
+	FGraphVertexHandle GetCompleteNodeHandle(const FGraphVertexHandle& InHandle) const;
 
 	/** Create a node with the specified subclass, adds it to the graph, and returns a handle to it. */
-	FGraphVertexHandle CreateVertex(int64 InUniqueIndex = INDEX_NONE);
+	FGraphVertexHandle CreateVertex(FGraphUniqueIndex InUniqueIndex = FGraphUniqueIndex());
 
 	/** Creates an edge between the two nodes. */
-	FGraphEdgeHandle CreateEdge(FGraphVertexHandle Node1, FGraphVertexHandle Node2, int64 InUniqueIndex = INDEX_NONE, bool bAddToIslands = true);
+	FGraphEdgeHandle CreateEdge(FGraphVertexHandle Node1, FGraphVertexHandle Node2, FGraphUniqueIndex InUniqueIndex = FGraphUniqueIndex(), bool bAddToIslands = true);
 
 	/** Creates edges in bulk. This is more efficient than calling CreateEdge multiple times since we will only try to assign a node to an island once. */
 	void CreateBulkEdges(TArray<TPair<FGraphVertexHandle, FGraphVertexHandle>>&& NodesToConnect);
@@ -141,18 +141,8 @@ private:
 	TMap<FGraphIslandHandle, TObjectPtr<UGraphIsland>> Islands;
 	FGraphProperties Properties;
 
-	/**
-	 *  The unique index to assign to the next node that gets created.
-	 *  If loading from persistence, as nodes get added into the graph, this
-	 *  value will automatically keep increasing so we never try to reuse
-	 *  the same unique index.
-	 */
-	int64 NextAvailableVertexUniqueIndex = 0;
-	int64 NextAvailableEdgeUniqueIndex = 0;
-	int64 NextAvailableIslandUniqueIndex = 0;
-
 	/** Creates an island out of a given set of nodes. */
-	FGraphIslandHandle CreateIsland(TArray<FGraphVertexHandle> Nodes, int64 InUniqueIndex = INDEX_NONE);
+	FGraphIslandHandle CreateIsland(TArray<FGraphVertexHandle> Nodes, FGraphUniqueIndex InUniqueIndex = FGraphUniqueIndex());
 
 	/** Adds a node to the graph's node collection and modifies the NextAvailableNodeUniqueIndex as necessary to maintain its validity. */
 	void RegisterVertex(TObjectPtr<UGraphVertex> Node);
diff --git a/Engine/Plugins/Experimental/GameplayGraph/Source/GameplayGraph/Public/Graph/GraphElement.h b/Engine/Plugins/Experimental/GameplayGraph/Source/GameplayGraph/Public/Graph/GraphElement.h
--- a/Engine/Plugins/Experimental/GameplayGraph/Source/GameplayGraph/Public/Graph/GraphElement.h
+++ b/Engine/Plugins/Experimental/GameplayGraph/Source/GameplayGraph/Public/Graph/GraphElement.h
@@ -29,8 +29,8 @@ public:
 protected:
 	UGraphElement();
 
-	void SetUniqueIndex(int64 InUniqueIndex) { UniqueIndex = InUniqueIndex; }
-	int64 GetUniqueIndex() const { return UniqueIndex; }
+	void SetUniqueIndex(FGraphUniqueIndex InUniqueIndex) { UniqueIndex = InUniqueIndex; }
+	FGraphUniqueIndex GetUniqueIndex() const { return UniqueIndex; }
 
 	void SetParentGraph(TObjectPtr<UGraph> InGraph);
 	TObjectPtr<UGraph> GetGraph() const;
@@ -44,7 +44,7 @@ private:
 
 	/** Will match the UniqueIndex in the UGraphHandle that references this element. */
 	UPROPERTY()
-	int64 UniqueIndex = INDEX_NONE;
+	FGraphUniqueIndex UniqueIndex = FGraphUniqueIndex();
 
 	UPROPERTY()
 	TWeakObjectPtr<UGraph> ParentGraph;
diff --git a/Engine/Plugins/Experimental/GameplayGraph/Source/GameplayGraph/Public/Graph/GraphHandle.h b/Engine/Plugins/Experimental/GameplayGraph/Source/GameplayGraph/Public/Graph/GraphHandle.h
--- a/Engine/Plugins/Experimental/GameplayGraph/Source/GameplayGraph/Public/Graph/GraphHandle.h
+++ b/Engine/Plugins/Experimental/GameplayGraph/Source/GameplayGraph/Public/Graph/GraphHandle.h
@@ -1,7 +1,9 @@
 // Copyright Epic Games, Inc. All Rights Reserved.
 #pragma once
 
+#include "Misc/Guid.h"
 #include "UObject/Object.h"
+
 #include "GraphHandle.generated.h"
 
 class UGraphElement;
@@ -11,6 +13,85 @@ class UGraphIsland;
 
 DECLARE_LOG_CATEGORY_EXTERN(LogGameplayGraph, Log, All);
 
+USTRUCT()
+struct GAMEPLAYGRAPH_API FGraphUniqueIndex
+{
+	GENERATED_BODY()
+public:
+	FGraphUniqueIndex(bool InIsTemp = false)
+		: UniqueIndex(FGuid())
+		, bIsTemporary(InIsTemp) 
+	{}
+
+	FGraphUniqueIndex(FGuid InUniqueIndex, bool InIsTemp = false)
+		: UniqueIndex(InUniqueIndex)
+		, bIsTemporary(InIsTemp) 
+	{}
+
+
+	bool IsValid() const
+	{
+		return UniqueIndex.IsValid();
+	}
+
+	bool IsTemporary() const
+	{
+		return bIsTemporary;
+	}
+
+	void SetTemporary(bool InTemp)
+	{
+		bIsTemporary = InTemp;
+	}
+
+	FGraphUniqueIndex NextUniqueIndex() const
+	{
+		return FGraphUniqueIndex(FGuid::NewGuid(),bIsTemporary);
+	}
+
+	static FGraphUniqueIndex CreateUniqueIndex(bool InIsTemp = false)
+	{
+		return FGraphUniqueIndex(FGuid::NewGuid(), InIsTemp);
+	}
+
+	bool operator==(const FGraphUniqueIndex& Other) const
+	{
+		return UniqueIndex == Other.UniqueIndex;
+	}
+
+	bool operator!=(const FGraphUniqueIndex& Other) const
+	{
+		return UniqueIndex != Other.UniqueIndex;
+	}
+
+	bool operator<(const FGraphUniqueIndex& Other) const
+	{
+		return UniqueIndex < Other.UniqueIndex;
+	}
+
+	friend uint32 GAMEPLAYGRAPH_API GetTypeHash(const FGraphUniqueIndex& InUniqueIndex)
+	{
+		return uint32(CityHash64((char*)&InUniqueIndex.UniqueIndex, sizeof(FGuid)));
+	}
+
+	FString ToString() const
+	{
+		return UniqueIndex.ToString();
+	}
+
+
+private:
+	/** Unique identifier within a graph. */
+	UPROPERTY(SaveGame)
+	FGuid UniqueIndex = FGuid();
+
+	/** Temporary Status for index */
+	UPROPERTY(Transient)
+	bool bIsTemporary = false;
+
+};
+
+
 /**
  * For persistence, every node in a graph is given a unique index.
  * A FGraphHandle encapsulates that index to make it easy to go from
@@ -22,7 +103,7 @@ struct GAMEPLAYGRAPH_API FGraphHandle
 	GENERATED_BODY()
 public:
 	FGraphHandle();
-	FGraphHandle(int64 InUniqueIndex, TObjectPtr<UGraphElement> InElement);
+	FGraphHandle(FGraphUniqueIndex InUniqueIndex, TObjectPtr<UGraphElement> InElement);
 
 	void Clear();
 
@@ -31,7 +112,7 @@ public:
 	bool HasElement() const;
 	bool IsComplete() const;
 
-	int64 GetUniqueIndex() const { return UniqueIndex; }
+	FGraphUniqueIndex GetUniqueIndex() const { return UniqueIndex; }
 
 	void SetElement(TObjectPtr<UGraphElement> InElement);
 	TObjectPtr<UGraphElement> GetElement() const;
@@ -44,7 +125,7 @@ public:
 private:
 	/** Unique identifier within a graph. */
 	UPROPERTY(SaveGame)
-	int64 UniqueIndex = INDEX_NONE;
+	FGraphUniqueIndex UniqueIndex = FGraphUniqueIndex();
 
 	/** Pointer to the graph */
 	UPROPERTY(Transient)
@@ -57,7 +138,7 @@ struct GAMEPLAYGRAPH_API FGraphVertexHandle : public FGraphHandle
 	GENERATED_BODY()
 
 	FGraphVertexHandle();
-	FGraphVertexHandle(int64 InUniqueIndex, TObjectPtr<UGraphElement> InElement)
+	FGraphVertexHandle(FGraphUniqueIndex InUniqueIndex, TObjectPtr<UGraphElement> InElement = nullptr)
 		: FGraphHandle(InUniqueIndex, InElement)
 	{}
 
@@ -70,7 +151,7 @@ struct GAMEPLAYGRAPH_API FGraphEdgeHandle : public FGraphHandle
 	GENERATED_BODY()
 
 	FGraphEdgeHandle();
-	FGraphEdgeHandle(int64 InUniqueIndex, TObjectPtr<UGraphElement> InElement)
+	FGraphEdgeHandle(FGraphUniqueIndex InUniqueIndex, TObjectPtr<UGraphElement> InElement)
 		: FGraphHandle(InUniqueIndex, InElement)
 	{}
 
@@ -83,7 +164,7 @@ struct GAMEPLAYGRAPH_API FGraphIslandHandle : public FGraphHandle
 	GENERATED_BODY()
 
 	FGraphIslandHandle();
-	FGraphIslandHandle(int64 InUniqueIndex, TObjectPtr<UGraphElement> InElement)
+	FGraphIslandHandle(FGraphUniqueIndex InUniqueIndex, TObjectPtr<UGraphElement> InElement)
 		: FGraphHandle(InUniqueIndex, InElement)
 	{}
 
