add the concept of allowed/disallowed operations to graph islands that are respected by the graph

diff --git a/Engine/Plugins/Experimental/GameplayGraph/Source/GameplayGraph/Private/Graph/Graph.cpp b/Engine/Plugins/Experimental/GameplayGraph/Source/GameplayGraph/Private/Graph/Graph.cpp
--- a/Engine/Plugins/Experimental/GameplayGraph/Source/GameplayGraph/Private/Graph/Graph.cpp
+++ b/Engine/Plugins/Experimental/GameplayGraph/Source/GameplayGraph/Private/Graph/Graph.cpp
@@ -171,7 +171,34 @@ FGraphEdgeHandle UGraph::CreateEdge(FGraphVertexHandle Node1, FGraphVertexHandle
 
 	if (!ensure(InUniqueIndex.IsValid()))
 	{
-		return{};
+		return {};
+	}
+
+	if (Properties.bGenerateIslands)
+	{
+		const FGraphIslandHandle& Island1 = Node1Ptr->GetParentIsland();
+		const FGraphIslandHandle& Island2 = Node2Ptr->GetParentIsland();
+
+		const TObjectPtr<UGraphIsland> Island1Ptr = Island1.GetIsland();
+		const TObjectPtr<UGraphIsland> Island2Ptr = Island2.GetIsland();
+
+		if (Island1Ptr && Island2Ptr)
+		{
+			// Possible merge scenario.
+			const bool bIsMerge = Island1 != Island2;
+			if (bIsMerge && !Island1Ptr->IsOperationAllowed(EGraphIslandOperations::Merge) && !Island2Ptr->IsOperationAllowed(EGraphIslandOperations::Merge))
+			{
+				return {};
+			}
+		}
+		else if (TObjectPtr<UGraphIsland> RelevantIsland = Island1Ptr ? Island1Ptr : Island2Ptr)
+		{
+			// Regular add scenario.
+			if (!RelevantIsland->IsOperationAllowed(EGraphIslandOperations::Add))
+			{
+				return {};
+			}
+		}
 	}
 
 	Edge->OnCreate();
@@ -278,6 +305,10 @@ void UGraph::RemoveIsland(const FGraphIslandHandle& IslandHandle)
 
 	if (TObjectPtr<UGraphIsland> Island = IslandHandle.GetIsland())
 	{
+		if (!Island->IsOperationAllowed(EGraphIslandOperations::Destroy))
+		{
+			return;
+		}
 		Island->Destroy();
 	}
 	Islands.Remove(IslandHandle);
@@ -513,21 +544,26 @@ void UGraph::RemoveBulkVertices(const TArray<FGraphVertexHandle>& InHandles)
 	{
 		if (NodeHandle.IsValid())
 		{
-			// We must remove every edge this node is a part of.
-			for (const FGraphEdgeHandle& EdgeHandle : VertexEdges.FindOrAdd(NodeHandle))
-			{
-				// Don't immediately handle islands. We'll do it later.
-				RemoveEdge(EdgeHandle, false);
-			}
-
 			if (TObjectPtr<UGraphVertex> Node = NodeHandle.GetVertex())
 			{
 				if (TObjectPtr<UGraphIsland> Island = Node->GetParentIsland().GetIsland())
 				{
+					if (!Island->IsOperationAllowed(EGraphIslandOperations::Remove))
+					{
+						continue;
+					}
+
 					AffectedIslands.Add(Node->GetParentIsland());
 					Island->RemoveVertex(NodeHandle);
 				}
 			}
+
+			// We must remove every edge this node is a part of.
+			for (const FGraphEdgeHandle& EdgeHandle : VertexEdges.FindOrAdd(NodeHandle))
+			{
+				// Don't immediately handle islands. We'll do it later.
+				RemoveEdge(EdgeHandle, false);
+			}
 		}
 	}
 	
@@ -608,7 +644,7 @@ void UGraph::RemoveEdge(const FGraphEdgeHandle& EdgeHandle, bool bHandleIslands)
 void UGraph::RemoveOrSplitIsland(TObjectPtr<UGraphIsland> Island)
 {
 	TRACE_CPUPROFILER_EVENT_SCOPE(UGraph::RemoveOrSplitIsland);
-	if (!Island)
+	if (!Island || !Island->IsOperationAllowed(EGraphIslandOperations::Split))
 	{
 		return;
 	}
@@ -723,4 +759,10 @@ void UGraph::FinalizeVertex(const FGraphVertexHandle& InHandle)
 	{
 		Island->HandleOnConnectivityChanged();
 	}
+}
+
+void UGraph::RefreshIslandConnectivity(const FGraphIslandHandle& IslandHandle)
+{
+	TRACE_CPUPROFILER_EVENT_SCOPE(UGraph::RefreshIslandConnectivity);
+	RemoveOrSplitIsland(IslandHandle.GetIsland());
 }
\ No newline at end of file
diff --git a/Engine/Plugins/Experimental/GameplayGraph/Source/GameplayGraph/Private/Graph/GraphIsland.cpp b/Engine/Plugins/Experimental/GameplayGraph/Source/GameplayGraph/Private/Graph/GraphIsland.cpp
--- a/Engine/Plugins/Experimental/GameplayGraph/Source/GameplayGraph/Private/Graph/GraphIsland.cpp
+++ b/Engine/Plugins/Experimental/GameplayGraph/Source/GameplayGraph/Private/Graph/GraphIsland.cpp
@@ -108,4 +108,16 @@ void UGraphIsland::HandleOnDestroyed()
 void UGraphIsland::HandleOnConnectivityChanged()
 {
 	OnConnectivityChanged.Broadcast(Handle());
+}
+
+void UGraphIsland::SetOperationAllowed(EGraphIslandOperations Op, bool bAllowed)
+{
+	if (bAllowed)
+	{
+		EnumAddFlags(AllowedOperations, Op);
+	}
+	else
+	{
+		EnumRemoveFlags(AllowedOperations, Op);
+	}
 }
\ No newline at end of file
diff --git a/Engine/Plugins/Experimental/GameplayGraph/Source/GameplayGraph/Public/Graph/Graph.h b/Engine/Plugins/Experimental/GameplayGraph/Source/GameplayGraph/Public/Graph/Graph.h
--- a/Engine/Plugins/Experimental/GameplayGraph/Source/GameplayGraph/Public/Graph/Graph.h
+++ b/Engine/Plugins/Experimental/GameplayGraph/Source/GameplayGraph/Public/Graph/Graph.h
@@ -113,6 +113,9 @@ public:
 	/** Remove an island from the graph. */
 	void RemoveIsland(const FGraphIslandHandle& IslandHandle);
 
+	/** Refresh the connectivity of the given island (re-check to see whether it should be split). */
+	void RefreshIslandConnectivity(const FGraphIslandHandle& IslandHandle);
+
 	int32 NumVertices() const { return Vertices.Num(); }
 	int32 NumEdges() const { return Edges.Num(); }
 	int32 NumIslands() const { return Islands.Num(); }
diff --git a/Engine/Plugins/Experimental/GameplayGraph/Source/GameplayGraph/Public/Graph/GraphIsland.h b/Engine/Plugins/Experimental/GameplayGraph/Source/GameplayGraph/Public/Graph/GraphIsland.h
--- a/Engine/Plugins/Experimental/GameplayGraph/Source/GameplayGraph/Public/Graph/GraphIsland.h
+++ b/Engine/Plugins/Experimental/GameplayGraph/Source/GameplayGraph/Public/Graph/GraphIsland.h
@@ -1,6 +1,7 @@
 // Copyright Epic Games, Inc. All Rights Reserved.
 #pragma once
 
+#include "Misc/EnumClassFlags.h"
 #include "Graph/GraphElement.h"
 #include "GraphIsland.generated.h"
 
@@ -13,6 +14,24 @@ struct FSerializedIslandData
 	TArray<FGraphVertexHandle> Vertices;
 };
 
+/**
+ * These are the possible operations that can be done to an island. The graph
+ * will attempt to check that the island is allowing these operations before
+ * successfully performing any of these operations. By default all operations are allowed.
+ */
+UENUM()
+enum class EGraphIslandOperations : int32
+{
+	None = 0,
+	Add = 1 << 0,
+	Remove = 1 << 1,
+	Split = 1 << 2,
+	Merge = 1 << 3,
+	Destroy = 1 << 4,
+	All = Add | Remove | Split | Merge | Destroy
+};
+ENUM_CLASS_FLAGS(EGraphIslandOperations);
+
 /** Delegate to track when some sort of batch change has occurred on this island that probably changes its connectivity.
  *  This is different from FOnGraphIslandNodeRemoved since FOnGraphIslandDestructiveChangeFinish will only be called
  *  once for the graph for a given operation while FOnGraphIslandNodeRemoved may be called multiple times if we're removing
@@ -47,6 +66,9 @@ public:
 	int32 Num() const { return Vertices.Num(); }
 	FSerializedIslandData GetSerializedData() const;
 
+	bool IsOperationAllowed(EGraphIslandOperations Op) const { return EnumHasAnyFlags(AllowedOperations, Op); }
+	void SetOperationAllowed(EGraphIslandOperations Op, bool bAllowed);
+
 	template<typename TLambda>
 	void ForEachVertex(TLambda&& Lambda)
 	{
@@ -84,4 +106,7 @@ private:
 
 	UPROPERTY(Transient)
 	bool bPendingDestroy;
+
+	UPROPERTY(Transient)
+	EGraphIslandOperations AllowedOperations = EGraphIslandOperations::All;
 };
\ No newline at end of file
