Preallocated some arrays used for generating a serializable connectivity graph to help with performance.

diff --git a/Engine/Plugins/Experimental/GameplayGraph/Source/GameplayGraph/Private/Graph/Graph.cpp b/Engine/Plugins/Experimental/GameplayGraph/Source/GameplayGraph/Private/Graph/Graph.cpp
--- a/Engine/Plugins/Experimental/GameplayGraph/Source/GameplayGraph/Private/Graph/Graph.cpp
+++ b/Engine/Plugins/Experimental/GameplayGraph/Source/GameplayGraph/Private/Graph/Graph.cpp
@@ -12,6 +12,9 @@ FSerializableGraph UGraph::GetSerializableGraph() const
 {
 	FSerializableGraph Out;
 	Out.Properties = Properties;
+	Out.Vertices.Reserve(Vertices.Num());
+	Out.Edges.Reserve(Edges.Num());
+	Out.Islands.Reserve(Islands.Num());
 
 	for (const TPair<FGraphVertexHandle, TObjectPtr<UGraphVertex>>& Kvp : Vertices)
 	{
