Comparison operators for GameplayGraph Graph, GraphEdge, and GraphIsland
NOTE: Moved to 28.10 from https://p4-swarm.epicgames.net/reviews/29937710 (engine code)

diff --git a/Engine/Plugins/Experimental/GameplayGraph/Source/GameplayGraph/Public/Graph/Graph.h b/Engine/Plugins/Experimental/GameplayGraph/Source/GameplayGraph/Public/Graph/Graph.h
--- a/Engine/Plugins/Experimental/GameplayGraph/Source/GameplayGraph/Public/Graph/Graph.h
+++ b/Engine/Plugins/Experimental/GameplayGraph/Source/GameplayGraph/Public/Graph/Graph.h
@@ -21,6 +21,10 @@ struct FGraphProperties
 
 	UPROPERTY(SaveGame)
 	bool bGenerateIslands = true;
+	
+	// Comparison operators
+	friend bool operator==(const FGraphProperties& Lhs, const FGraphProperties& Rhs) = default;
+	friend bool operator!=(const FGraphProperties& Lhs, const FGraphProperties& Rhs) = default;
 };
 
 DECLARE_MULTICAST_DELEGATE_OneParam(FOnGraphVertexCreated, const FGraphVertexHandle&);
diff --git a/Engine/Plugins/Experimental/GameplayGraph/Source/GameplayGraph/Public/Graph/GraphEdge.h b/Engine/Plugins/Experimental/GameplayGraph/Source/GameplayGraph/Public/Graph/GraphEdge.h
--- a/Engine/Plugins/Experimental/GameplayGraph/Source/GameplayGraph/Public/Graph/GraphEdge.h
+++ b/Engine/Plugins/Experimental/GameplayGraph/Source/GameplayGraph/Public/Graph/GraphEdge.h
@@ -16,6 +16,10 @@ struct FSerializedEdgeData
 
 	UPROPERTY(SaveGame)
 	FGraphVertexHandle Node2;
+
+	// Comparison operators
+	friend bool operator==(const FSerializedEdgeData& Lhs, const FSerializedEdgeData& Rhs) = default;
+	friend bool operator!=(const FSerializedEdgeData& Lhs, const FSerializedEdgeData& Rhs) = default;
 };
 
 UCLASS()
diff --git a/Engine/Plugins/Experimental/GameplayGraph/Source/GameplayGraph/Public/Graph/GraphIsland.h b/Engine/Plugins/Experimental/GameplayGraph/Source/GameplayGraph/Public/Graph/GraphIsland.h
--- a/Engine/Plugins/Experimental/GameplayGraph/Source/GameplayGraph/Public/Graph/GraphIsland.h
+++ b/Engine/Plugins/Experimental/GameplayGraph/Source/GameplayGraph/Public/Graph/GraphIsland.h
@@ -12,6 +12,10 @@ struct FSerializedIslandData
 
 	UPROPERTY(SaveGame)
 	TArray<FGraphVertexHandle> Vertices;
+
+	// Comparison operators
+	friend bool operator==(const FSerializedIslandData& Lhs, const FSerializedIslandData& Rhs) = default;
+	friend bool operator!=(const FSerializedIslandData& Lhs, const FSerializedIslandData& Rhs) = default;
 };
 
 /**
