fix some potential issues in the gameplay graph caused by mismatches between the vertices in the island set vs what is actually in the graph

- removed the ability to disable Remove operations on islands
- sanity check the island vertex set when loading from serialized data

diff --git a/Engine/Plugins/Experimental/GameplayGraph/Source/GameplayGraph/Private/Graph/Graph.cpp b/Engine/Plugins/Experimental/GameplayGraph/Source/GameplayGraph/Private/Graph/Graph.cpp
--- a/Engine/Plugins/Experimental/GameplayGraph/Source/GameplayGraph/Private/Graph/Graph.cpp
+++ b/Engine/Plugins/Experimental/GameplayGraph/Source/GameplayGraph/Private/Graph/Graph.cpp
@@ -84,10 +84,16 @@ void UGraph::LoadFromSerializedGraph(const FSerializableGraph& Input)
 			continue;
 		}
 
-		TArray<FGraphVertexHandle> IslandVertices = Kvp.Value.Vertices;
-		for (FGraphVertexHandle& VertexHandle : IslandVertices)
+		TArray<FGraphVertexHandle> IslandVertices;
+		IslandVertices.Reserve(Kvp.Value.Vertices.Num());
+		
+		for (const FGraphVertexHandle& VertexHandle : Kvp.Value.Vertices)
 		{
-			VertexHandle = GetCompleteNodeHandle(VertexHandle);
+			FGraphVertexHandle CompleteHandle = GetCompleteNodeHandle(VertexHandle);
+			if (CompleteHandle.IsComplete())
+			{
+				IslandVertices.Add(CompleteHandle);
+			}
 		}
 
 		CreateIsland(IslandVertices, Kvp.Key.GetUniqueIndex());
@@ -551,11 +557,6 @@ void UGraph::RemoveBulkVertices(const TArray<FGraphVertexHandle>& InHandles)
 			{
 				if (TObjectPtr<UGraphIsland> Island = Node->GetParentIsland().GetIsland())
 				{
-					if (!Island->IsOperationAllowed(EGraphIslandOperations::Remove))
-					{
-						continue;
-					}
-
 					AffectedIslands.Add(Node->GetParentIsland());
 					Island->RemoveVertex(NodeHandle);
 				}
diff --git a/Engine/Plugins/Experimental/GameplayGraph/Source/GameplayGraph/Public/Graph/GraphIsland.h b/Engine/Plugins/Experimental/GameplayGraph/Source/GameplayGraph/Public/Graph/GraphIsland.h
--- a/Engine/Plugins/Experimental/GameplayGraph/Source/GameplayGraph/Public/Graph/GraphIsland.h
+++ b/Engine/Plugins/Experimental/GameplayGraph/Source/GameplayGraph/Public/Graph/GraphIsland.h
@@ -28,11 +28,10 @@ enum class EGraphIslandOperations : int32
 {
 	None = 0,
 	Add = 1 << 0,
-	Remove = 1 << 1,
-	Split = 1 << 2,
-	Merge = 1 << 3,
-	Destroy = 1 << 4,
-	All = Add | Remove | Split | Merge | Destroy
+	Split = 1 << 1,
+	Merge = 1 << 2,
+	Destroy = 1 << 3,
+	All = Add | Split | Merge | Destroy
 };
 ENUM_CLASS_FLAGS(EGraphIslandOperations);
 
