[Backout] - CL31192713
Original CL Desc
-----------------------------------------------------------------
gameplay graph tests and some fixes stemming from these tests

- UGraph::GetCompleteNodeHandle will return an empty handle if the index is invalid but has a valid element (or in any situtation where the index does not exist in the graph)
- Made UGraph::CreateEdge private to force usage of UGraph::CreateBulkEdges
- UGraph::CreateBulkEdges can specifiy output edges as well as edge unique index using the FEdgeCreationParameters struct
- UGraphEdge::GetOtherNode will only return a valid node handle if the input node is one of the two nodes in the edge
- UGraph::CreateVertex can create a vertex with a random valid unique index if not otherwise given one
- Remove extra UGraph::VertexEdges book-keeping
- Add UGraphVertex::NumEdges and UGraphVertex::GetEdges
- UGraphHandle::GetElement is removed - should use the typed functions on the vertex/edge/island handles instead
- Made UGraph::GetVertices, UGraph::GetEdges, UGraph::GetIslands, and UGraph::GetProperties public
- FGraphHandle now does not cache a pointer to the element but a pointer to the graph to ensure that once the vertex/edge/island is removed from the graph, every handle becomes incomplete. This also ensures that UGraph::CreateVertex can't create a duplicate vertex handle that points to a different underlying vertex element.
- Remove bHandleIslands from the public UGraph::RemoveEdge to remove a way for users to put the graph in a weird state
- Remove UGraphIsland::bPendingDestroy as it wasn't used
- Fixed an issue where disabling EGraphIslandOperations::Merge on a single island would not end up actually preventing the merge
- Creating an edge with UGraph::CreateEdge will always initialize the a node's island if it didn't exist before.
- UGraph::LoadFromSerializedGraph will now merge/split islands depending on the edge connectivity

#tests low level tests

diff --git a/Engine/Plugins/Experimental/GameplayGraph/Source/GameplayGraph/Private/Graph/Graph.cpp b/Engine/Plugins/Experimental/GameplayGraph/Source/GameplayGraph/Private/Graph/Graph.cpp
--- a/Engine/Plugins/Experimental/GameplayGraph/Source/GameplayGraph/Private/Graph/Graph.cpp
+++ b/Engine/Plugins/Experimental/GameplayGraph/Source/GameplayGraph/Private/Graph/Graph.cpp
@@ -5,17 +5,9 @@
 #include "Graph/GraphVertex.h"
 #include "Graph/Algorithms/Connectivity/ConnectedComponents.h"
 #include "GenericPlatform/GenericPlatformMath.h"
-#include "HAL/IConsoleManager.h"
 
 #include UE_INLINE_GENERATED_CPP_BY_NAME(Graph)
 
-
-namespace Graph
-{
-	static bool bFixupGraphOnLoad = true;
-	static FAutoConsoleVariableRef FixupGraphOnLoadCVar(TEXT("GameplayGraph.FixupGraphOnLoad"), bFixupGraphOnLoad, TEXT("Merge and split islands when loaded from serialization."));
-}
-
 FSerializableGraph UGraph::GetSerializableGraph() const
 {
 	FSerializableGraph Out;
@@ -70,8 +62,6 @@ void UGraph::LoadFromSerializedGraph(const FSerializableGraph& Input)
 		CreateVertex(Handle.GetUniqueIndex());
 	}
 
-	TArray<FGraphEdgeHandle> AllEdges;
-	AllEdges.Reserve(Input.Edges.Num());
 	for (const TPair<FGraphEdgeHandle, FSerializedEdgeData>& Kvp : Input.Edges)
 	{
 		if (!Kvp.Key.IsValid())
@@ -79,16 +69,14 @@ void UGraph::LoadFromSerializedGraph(const FSerializableGraph& Input)
 			continue;
 		}
 
-		AllEdges.Emplace(CreateEdge(
+		CreateEdge(
 			GetCompleteNodeHandle(Kvp.Value.Node1),
 			GetCompleteNodeHandle(Kvp.Value.Node2),
 			Kvp.Key.GetUniqueIndex(),
 			false
-		));
+		);
 	}
 
-	TArray<FGraphIslandHandle> AllIslands;
-	AllIslands.Reserve(Input.Islands.Num());
 	for (const TPair<FGraphIslandHandle, FSerializedIslandData>& Kvp : Input.Islands)
 	{
 		if (!Kvp.Key.IsValid())
@@ -108,17 +96,7 @@ void UGraph::LoadFromSerializedGraph(const FSerializableGraph& Input)
 			}
 		}
 
-		AllIslands.Emplace(CreateIsland(IslandVertices, Kvp.Key.GetUniqueIndex()));
-	}
-
-	if (Graph::bFixupGraphOnLoad)
-	{
-		MergeOrCreateIslands(AllEdges);
-
-		for (const FGraphIslandHandle& NewIslandHandle : AllIslands)
-		{
-			RemoveOrSplitIsland(NewIslandHandle.GetIsland());
-		}
+		CreateIsland(IslandVertices, Kvp.Key.GetUniqueIndex());
 	}
 }
 
@@ -146,17 +124,10 @@ FGraphVertexHandle UGraph::CreateVertex(FGraphUniqueIndex InUniqueIndex)
 
 	if (!ensure(InUniqueIndex.IsValid()))
 	{
-		return {};
+		return{};
 	}
-
-	Vertex->SetUniqueIndex(InUniqueIndex);
-	Vertex->SetParentGraph(this);
-	if (Vertices.Contains(Vertex->Handle()))
-	{
-		return {};
-	}
-
 	Vertex->OnCreate();
+	Vertex->SetUniqueIndex(InUniqueIndex);
 	RegisterVertex(Vertex);
 	OnVertexCreated.Broadcast(Vertex->Handle());
 	return Vertex->Handle();
@@ -170,12 +141,12 @@ void UGraph::RegisterVertex(TObjectPtr<UGraphVertex> Vertex)
 		return;
 	}
 
-	Vertex->SetParentGraph(this);
 	const FGraphVertexHandle Handle = Vertex->Handle();
+	Vertex->SetParentGraph(this);
 	Vertices.Add(Handle, Vertex);
 }
 
-FGraphEdgeHandle UGraph::CreateEdge(FGraphVertexHandle Node1, FGraphVertexHandle Node2, FGraphUniqueIndex InUniqueIndex, bool bMergeIslands)
+FGraphEdgeHandle UGraph::CreateEdge(FGraphVertexHandle Node1, FGraphVertexHandle Node2, FGraphUniqueIndex InUniqueIndex, bool bAddToIslands)
 {
 	TRACE_CPUPROFILER_EVENT_SCOPE(UGraph::CreateEdge);
 	if (!Node1.IsComplete() || !Node2.IsComplete())
@@ -224,7 +195,7 @@ FGraphEdgeHandle UGraph::CreateEdge(FGraphVertexHandle Node1, FGraphVertexHandle
 		{
 			// Possible merge scenario.
 			const bool bIsMerge = Island1 != Island2;
-			if (bIsMerge && (!Island1Ptr->IsOperationAllowed(EGraphIslandOperations::Merge) || !Island2Ptr->IsOperationAllowed(EGraphIslandOperations::Merge)))
+			if (bIsMerge && !Island1Ptr->IsOperationAllowed(EGraphIslandOperations::Merge) && !Island2Ptr->IsOperationAllowed(EGraphIslandOperations::Merge))
 			{
 				return {};
 			}
@@ -252,44 +223,23 @@ FGraphEdgeHandle UGraph::CreateEdge(FGraphVertexHandle Node1, FGraphVertexHandle
 	OnEdgeCreated.Broadcast(Edge->Handle());
 
 	// If we want to keep track of islands, this is where we need to create/merge islands.
-	if (Properties.bGenerateIslands)
+	if (Properties.bGenerateIslands && bAddToIslands)
 	{
-		// After we've created the edge - we want to see one of the two islands is invalid in which
-		// case we have to initialize that node's island. This is needed or else we might not detect a merge
-		// properly when creating multiple edges at a time in UGraph::CreateBulkEdges.
-		const FGraphIslandHandle& IslandHandle1 = Node1Ptr->GetParentIsland();
-		const FGraphIslandHandle& IslandHandle2 = Node2Ptr->GetParentIsland();
-
-		if (!IslandHandle1.IsValid() || !IslandHandle2.IsValid())
-		{
-			if (UGraphIsland* Island1 = IslandHandle1.GetIsland())
-			{
-				Island1->AddVertex(Node2);
-			}
-			else if (UGraphIsland* Island2 = IslandHandle1.GetIsland())
-			{
-				Island2->AddVertex(Node1);
-			}
-		}
-
-		if (bMergeIslands)
-		{
-			MergeOrCreateIslands({ Edge->Handle() });
-		}
+		MergeOrCreateIslands({ Edge->Handle() });
 	}
 
 	return Edge->Handle();
 }
 
-void UGraph::CreateBulkEdges(TArray<FEdgeCreationParameters>&& NodesToConnect, TArray<FGraphEdgeHandle>* OutEdges)
+void UGraph::CreateBulkEdges(TArray<TPair<FGraphVertexHandle, FGraphVertexHandle>>&& NodesToConnect)
 {
 	// Create all edges normally but don't call MergeOrCreateIslands yet. We'll use the bulk function instead.
 	TArray<FGraphEdgeHandle> NewEdges;
 	NewEdges.Reserve(NodesToConnect.Num());
 
-	for (const FEdgeCreationParameters& Params : NodesToConnect)
+	for (const TPair<FGraphVertexHandle, FGraphVertexHandle>& NodePair : NodesToConnect)
 	{
-		if (FGraphEdgeHandle Edge = CreateEdge(Params.VertexHandle1, Params.VertexHandle2, Params.EdgeIndex, false); Edge.IsValid())
+		if (FGraphEdgeHandle Edge = CreateEdge(NodePair.Key, NodePair.Value, FGraphUniqueIndex::CreateUniqueIndex(), false); Edge.IsValid())
 		{
 			NewEdges.Add(Edge);
 		}
@@ -297,12 +247,7 @@ void UGraph::CreateBulkEdges(TArray<FEdgeCreationParameters>&& NodesToConnect, T
 
 	// MergeOrCreateIslands incrementally determines island connectivity one edge at a time. This is efficient if we're handling a single edge but less
 	// efficient if we're trying to add a bunch of edges all at the same time since it'll cause a vertex to jump between islands. The Bulk function helps prevent that.
-	MergeOrCreateIslands(NewEdges);
-
-	if (OutEdges)
-	{
-		*OutEdges = MoveTemp(NewEdges);
-	}
+	MergeOrCreateIslands(MoveTemp(NewEdges));
 }
 
 void UGraph::RegisterEdge(TObjectPtr<UGraphEdge> Edge)
@@ -313,9 +258,12 @@ void UGraph::RegisterEdge(TObjectPtr<UGraphEdge> Edge)
 		return;
 	}
 
-	Edge->SetParentGraph(this);
 	const FGraphEdgeHandle Handle = Edge->Handle();
 	Edges.Add(Handle, Edge);
+	Edge->SetParentGraph(this);
+
+	VertexEdges.FindOrAdd(Edge->NodeA()).Add(Handle);
+	VertexEdges.FindOrAdd(Edge->NodeB()).Add(Handle);
 }
 
 FGraphIslandHandle UGraph::CreateIsland(TArray<FGraphVertexHandle> InputNodes, FGraphUniqueIndex InUniqueIndex)
@@ -331,10 +279,8 @@ FGraphIslandHandle UGraph::CreateIsland(TArray<FGraphVertexHandle> InputNodes, F
 	{
 		return {};
 	}
-	Island->SetParentGraph(this);
 	Island->OnCreate();
 	Island->SetUniqueIndex((!InUniqueIndex.IsValid()) ? FGraphUniqueIndex::CreateUniqueIndex() : InUniqueIndex);
-	RegisterIsland(Island);
 
 	OnIslandCreated.Broadcast(Island->Handle());
 	for (const FGraphVertexHandle& Node : InputNodes)
@@ -342,6 +288,7 @@ FGraphIslandHandle UGraph::CreateIsland(TArray<FGraphVertexHandle> InputNodes, F
 		Island->AddVertex(Node);
 	}
 
+	RegisterIsland(Island);
 	return Island->Handle();
 }
 
@@ -376,7 +323,7 @@ void UGraph::RemoveIsland(const FGraphIslandHandle& IslandHandle)
 	Islands.Remove(IslandHandle);
 }
 
-void UGraph::MergeOrCreateIslands(const TArray<FGraphEdgeHandle>& InEdges)
+void UGraph::MergeOrCreateIslands(TArray<FGraphEdgeHandle>&& InEdges)
 {
 	TRACE_CPUPROFILER_EVENT_SCOPE(UGraph::MergeOrCreateIslands);
 	if (Edges.IsEmpty())
@@ -613,25 +560,14 @@ void UGraph::RemoveBulkVertices(const TArray<FGraphVertexHandle>& InHandles)
 					AffectedIslands.Add(Node->GetParentIsland());
 					Island->RemoveVertex(NodeHandle);
 				}
-
-				// We must remove every edge this node is a part of. Need to make a copy of the edges because
-				// otherwise we're modifying the container during iteration over it in ForEachAdjacentVertex.
-				TArray<FGraphEdgeHandle> EdgeCopy;
-				EdgeCopy.Reserve(Node->NumEdges());
-				Node->ForEachAdjacentVertex(
-					[&EdgeCopy](const FGraphVertexHandle& OtherNodeHandle, const FGraphEdgeHandle& EdgeHandle)
-					{
-						EdgeCopy.Add(EdgeHandle);
-					}
-				);
-
-				for (const FGraphEdgeHandle& EdgeHandle : EdgeCopy)
-				{
-					// Don't immediately handle islands. We'll do it later.
-					RemoveEdgeInternal(EdgeHandle, false);
-				}
 			}
 
+			// We must remove every edge this node is a part of.
+			for (const FGraphEdgeHandle& EdgeHandle : VertexEdges.FindOrAdd(NodeHandle))
+			{
+				// Don't immediately handle islands. We'll do it later.
+				RemoveEdge(EdgeHandle, false);
+			}
 		}
 	}
 	
@@ -654,24 +590,22 @@ void UGraph::RemoveBulkVertices(const TArray<FGraphVertexHandle>& InHandles)
 				Node->HandleOnVertexRemoved();
 			}
 			Vertices.Remove(NodeHandle);
+			VertexEdges.Remove(NodeHandle);
 		}
 	}
 }
 
-void UGraph::RemoveEdge(const FGraphEdgeHandle& EdgeHandle)
-{
-	RemoveEdgeInternal(EdgeHandle, true);
-}
-
-void UGraph::RemoveEdgeInternal(const FGraphEdgeHandle& EdgeHandle, bool bHandleIslands)
+void UGraph::RemoveEdge(const FGraphEdgeHandle& EdgeHandle, bool bHandleIslands)
 {
-	TRACE_CPUPROFILER_EVENT_SCOPE(UGraph::RemoveEdgeInternal);
+	TRACE_CPUPROFILER_EVENT_SCOPE(UGraph::RemoveEdge);
 	if (!EdgeHandle.IsValid())
 	{
 		return;
 	}
 
 	TObjectPtr<UGraphEdge> Edge = EdgeHandle.GetEdge();
+	Edges.Remove(EdgeHandle);
+
 	if (!Edge)
 	{
 		return;
@@ -709,8 +643,6 @@ void UGraph::RemoveEdgeInternal(const FGraphEdgeHandle& EdgeHandle, bool bHandle
 		// Note that we can assume that both nodes are in the same island.
 		RemoveOrSplitIsland(IslandHandle.GetIsland());
 	}
-
-	Edges.Remove(EdgeHandle);
 }
 
 void UGraph::RemoveOrSplitIsland(TObjectPtr<UGraphIsland> Island)
@@ -777,9 +709,9 @@ void UGraph::RemoveOrSplitIsland(TObjectPtr<UGraphIsland> Island)
 
 FGraphVertexHandle UGraph::GetCompleteNodeHandle(const FGraphVertexHandle& InHandle) const
 {
-	if (!InHandle.IsValid())
+	if (!InHandle.IsValid() || InHandle.HasElement())
 	{
-		return {};
+		return InHandle;
 	}
 
 	const TObjectPtr<UGraphVertex>* Data = Vertices.Find(InHandle);
@@ -827,28 +759,6 @@ void UGraph::FinalizeVertex(const FGraphVertexHandle& InHandle)
 		IslandHandle = CreateIsland({ InHandle });
 	}
 
-	// Need to iterate through edges to see if any edge was added prior to vertex finalization.
-	// If so, we need to do a merge.
-	TArray<FGraphEdgeHandle> EdgesToMerge;
-	Node->ForEachAdjacentVertex(
-		[&IslandHandle, &EdgesToMerge](const FGraphVertexHandle& NeighborVertexHandle, const FGraphEdgeHandle& EdgeHandle)
-		{
-			if (UGraphVertex* NeighborVertex = NeighborVertexHandle.GetVertex())
-			{
-				const FGraphIslandHandle& NeighborIslandHandle = NeighborVertex->GetParentIsland();
-				if (NeighborIslandHandle.IsComplete() && NeighborIslandHandle != IslandHandle)
-				{
-					EdgesToMerge.Add(EdgeHandle);
-				}
-			}
-		}
-	);
-
-	if (!EdgesToMerge.IsEmpty())
-	{
-		MergeOrCreateIslands(EdgesToMerge);
-	}
-
 	if (TObjectPtr<UGraphIsland> Island = IslandHandle.GetIsland())
 	{
 		Island->HandleOnConnectivityChanged();
diff --git a/Engine/Plugins/Experimental/GameplayGraph/Source/GameplayGraph/Private/Graph/GraphEdge.cpp b/Engine/Plugins/Experimental/GameplayGraph/Source/GameplayGraph/Private/Graph/GraphEdge.cpp
--- a/Engine/Plugins/Experimental/GameplayGraph/Source/GameplayGraph/Private/Graph/GraphEdge.cpp
+++ b/Engine/Plugins/Experimental/GameplayGraph/Source/GameplayGraph/Private/Graph/GraphEdge.cpp
@@ -21,13 +21,9 @@ const FGraphVertexHandle& UGraphEdge::GetOtherNode(const FGraphVertexHandle& InN
 	{
 		return B;
 	}
-	else if (InNode == B)
-	{
-		return A;
-	}
 	else
 	{
-		return FGraphVertexHandle::Invalid;
+		return A;
 	}
 }
 
diff --git a/Engine/Plugins/Experimental/GameplayGraph/Source/GameplayGraph/Private/Graph/GraphHandle.cpp b/Engine/Plugins/Experimental/GameplayGraph/Source/GameplayGraph/Private/Graph/GraphHandle.cpp
--- a/Engine/Plugins/Experimental/GameplayGraph/Source/GameplayGraph/Private/Graph/GraphHandle.cpp
+++ b/Engine/Plugins/Experimental/GameplayGraph/Source/GameplayGraph/Private/Graph/GraphHandle.cpp
@@ -1,8 +1,6 @@
 // Copyright Epic Games, Inc. All Rights Reserved.
 
 #include "Graph/GraphHandle.h"
-
-#include "Graph/Graph.h"
 #include "Graph/GraphElement.h"
 #include "Graph/GraphEdge.h"
 #include "Graph/GraphIsland.h"
@@ -13,21 +11,37 @@
 
 DEFINE_LOG_CATEGORY(LogGameplayGraph)
 
-FGraphVertexHandle FGraphVertexHandle::Invalid;
-FGraphEdgeHandle FGraphEdgeHandle::Invalid;
-FGraphIslandHandle FGraphIslandHandle::Invalid;
+FGraphHandle::FGraphHandle()
+{
 
-FGraphHandle::FGraphHandle(FGraphUniqueIndex InUniqueIndex, UGraph* InGraph)
+}
+
+FGraphHandle::FGraphHandle(FGraphUniqueIndex InUniqueIndex, TObjectPtr<UGraphElement> InElement)
 	: UniqueIndex(InUniqueIndex)
-	, WeakGraph(InGraph)
+	, Element(InElement)
 {
 }
 
+void FGraphHandle::SetElement(TObjectPtr<UGraphElement> InElement)
+{
+	Element = InElement;
+}
+
+TObjectPtr<UGraphElement> FGraphHandle::GetElement() const
+{
+	return Element.Get();
+}
+
 bool FGraphHandle::IsValid() const
 {
 	return UniqueIndex.IsValid();
 }
 
+bool FGraphHandle::HasElement() const
+{
+	return Element.IsValid();
+}
+
 bool FGraphHandle::IsComplete() const
 {
 	return IsValid() && HasElement();
@@ -51,7 +65,7 @@ bool FGraphHandle::operator<(const FGraphHandle& Other) const
 void FGraphHandle::Clear()
 {
 	UniqueIndex = FGraphUniqueIndex();
-	WeakGraph = nullptr;
+	Element = nullptr;
 }
 
 uint32 GetTypeHash(const FGraphHandle& Handle)
@@ -59,62 +73,31 @@ uint32 GetTypeHash(const FGraphHandle& Handle)
 	return GetTypeHash(Handle.UniqueIndex);
 }
 
-UGraphVertex* FGraphVertexHandle::GetVertex() const
+FGraphVertexHandle::FGraphVertexHandle()
 {
-	if (UGraph* Graph = GetGraph())
-	{
-		return Graph->GetVertices().FindRef(*this);
-	}
-
-	return nullptr;
 }
 
-bool FGraphVertexHandle::HasElement() const
+TObjectPtr<UGraphVertex> FGraphVertexHandle::GetVertex() const
 {
-	if (UGraph* Graph = GetGraph())
-	{
-		return Graph->GetVertices().Contains(*this);
-	}
-
-	return false;
+	return Cast<UGraphVertex>(GetElement());
 }
 
-UGraphEdge* FGraphEdgeHandle::GetEdge() const
+FGraphEdgeHandle::FGraphEdgeHandle()
 {
-	if (UGraph* Graph = GetGraph())
-	{
-		return Graph->GetEdges().FindRef(*this);
-	}
 
-	return nullptr;
 }
 
-bool FGraphEdgeHandle::HasElement() const
+TObjectPtr<UGraphEdge> FGraphEdgeHandle::GetEdge() const
 {
-	if (UGraph* Graph = GetGraph())
-	{
-		return Graph->GetEdges().Contains(*this);
-	}
-
-	return false;
+	return Cast<UGraphEdge>(GetElement());
 }
 
-UGraphIsland* FGraphIslandHandle::GetIsland() const
+FGraphIslandHandle::FGraphIslandHandle()
 {
-	if (UGraph* Graph = GetGraph())
-	{
-		return Graph->GetIslands().FindRef(*this);
-	}
 
-	return nullptr;
 }
 
-bool FGraphIslandHandle::HasElement() const
+TObjectPtr<UGraphIsland> FGraphIslandHandle::GetIsland() const
 {
-	if (UGraph* Graph = GetGraph())
-	{
-		return Graph->GetIslands().Contains(*this);
-	}
-
-	return false;
+	return Cast<UGraphIsland>(GetElement());
 }
\ No newline at end of file
diff --git a/Engine/Plugins/Experimental/GameplayGraph/Source/GameplayGraph/Private/Graph/GraphIsland.cpp b/Engine/Plugins/Experimental/GameplayGraph/Source/GameplayGraph/Private/Graph/GraphIsland.cpp
--- a/Engine/Plugins/Experimental/GameplayGraph/Source/GameplayGraph/Private/Graph/GraphIsland.cpp
+++ b/Engine/Plugins/Experimental/GameplayGraph/Source/GameplayGraph/Private/Graph/GraphIsland.cpp
@@ -9,11 +9,13 @@
 UGraphIsland::UGraphIsland()
 	: Super(EGraphElementType::Island)
 {
+	bPendingDestroy = false;
 }
 
 void UGraphIsland::Destroy()
 {
 	TRACE_CPUPROFILER_EVENT_SCOPE(UGraphIsland::Destroy);
+	bPendingDestroy = true;
 	TSet<FGraphVertexHandle> VertexCopy = Vertices;
 	Vertices.Empty();
 
diff --git a/Engine/Plugins/Experimental/GameplayGraph/Source/GameplayGraph/Public/Graph/Graph.h b/Engine/Plugins/Experimental/GameplayGraph/Source/GameplayGraph/Public/Graph/Graph.h
--- a/Engine/Plugins/Experimental/GameplayGraph/Source/GameplayGraph/Public/Graph/Graph.h
+++ b/Engine/Plugins/Experimental/GameplayGraph/Source/GameplayGraph/Public/Graph/Graph.h
@@ -56,13 +56,6 @@ struct FSerializableGraph
 	TMap<FGraphIslandHandle, FSerializedIslandData> Islands;
 };
 
-struct FEdgeCreationParameters
-{
-	FGraphVertexHandle VertexHandle1;
-	FGraphVertexHandle VertexHandle2;
-	FGraphUniqueIndex EdgeIndex = FGraphUniqueIndex::CreateUniqueIndex();
-};
-
 /**
  * A UGraph is a collection of nodes and edges. This graph representation
  * is meant to be easily integrable into gameplay systems in the Unreal Engine.
@@ -102,10 +95,13 @@ public:
 	FGraphVertexHandle GetCompleteNodeHandle(const FGraphVertexHandle& InHandle) const;
 
 	/** Create a node with the specified subclass, adds it to the graph, and returns a handle to it. */
-	FGraphVertexHandle CreateVertex(FGraphUniqueIndex InUniqueIndex = FGraphUniqueIndex::CreateUniqueIndex(false));
+	FGraphVertexHandle CreateVertex(FGraphUniqueIndex InUniqueIndex = FGraphUniqueIndex());
+
+	/** Creates an edge between the two nodes. */
+	FGraphEdgeHandle CreateEdge(FGraphVertexHandle Node1, FGraphVertexHandle Node2, FGraphUniqueIndex InUniqueIndex = FGraphUniqueIndex(), bool bAddToIslands = true);
 
 	/** Creates edges in bulk. This is more efficient than calling CreateEdge multiple times since we will only try to assign a node to an island once. */
-	void CreateBulkEdges(TArray<FEdgeCreationParameters>&& NodesToConnect, TArray<FGraphEdgeHandle>* OutEdges = nullptr);
+	void CreateBulkEdges(TArray<TPair<FGraphVertexHandle, FGraphVertexHandle>>&& NodesToConnect);
 
 	/** Removes a node from the graph along with any edges that contain it. */
 	void RemoveVertex(const FGraphVertexHandle& NodeHandle);
@@ -113,7 +109,7 @@ public:
 	void RemoveBulkVertices(const TArray<FGraphVertexHandle>& InHandles);
 
 	/** Removes an edge from the graph. */
-	void RemoveEdge(const FGraphEdgeHandle& EdgeHandle);
+	void RemoveEdge(const FGraphEdgeHandle& EdgeHandle, bool bHandleIslands);
 
 	/** This should be called immediately after a node and any relevant edges have been added to the graph. */
 	virtual void FinalizeVertex(const FGraphVertexHandle& InHandle);
@@ -128,19 +124,11 @@ public:
 	int32 NumEdges() const { return Edges.Num(); }
 	int32 NumIslands() const { return Islands.Num(); }
 
-	template<typename TLambda>
-	void ForEachIsland(TLambda&& Lambda)
-	{
-		for (const TPair<FGraphIslandHandle, TObjectPtr<UGraphIsland>>& Kvp : Islands)
-		{
-			Lambda(Kvp.Key, Kvp.Value);
-		}
-	}
-
 	FOnGraphVertexCreated OnVertexCreated;
 	FOnGraphEdgeCreated OnEdgeCreated;
 	FOnGraphIslandCreated OnIslandCreated;
 
+protected:
 	const TMap<FGraphVertexHandle, TObjectPtr<UGraphVertex>>& GetVertices() const { return Vertices; }
 	const TMap<FGraphEdgeHandle, TObjectPtr<UGraphEdge>>& GetEdges() const { return Edges; }
 	const TMap<FGraphIslandHandle, TObjectPtr<UGraphIsland>>& GetIslands() const { return Islands; }
@@ -150,6 +138,9 @@ private:
 	UPROPERTY()
 	TMap<FGraphVertexHandle, TObjectPtr<UGraphVertex>> Vertices;
 
+	/** Keeps track of every edge a node is a part of. The node itself will only track edges it's the "source" node of in directed graphs. */
+	TMap<FGraphVertexHandle, TSet<FGraphEdgeHandle>> VertexEdges;
+
 	UPROPERTY()
 	TMap<FGraphEdgeHandle, TObjectPtr<UGraphEdge>> Edges;
 
@@ -160,9 +151,6 @@ private:
 	/** Creates an island out of a given set of nodes. */
 	FGraphIslandHandle CreateIsland(TArray<FGraphVertexHandle> Nodes, FGraphUniqueIndex InUniqueIndex = FGraphUniqueIndex());
 
-	/** Creates an edge between the two nodes. */
-	FGraphEdgeHandle CreateEdge(FGraphVertexHandle Node1, FGraphVertexHandle Node2, FGraphUniqueIndex InUniqueIndex = FGraphUniqueIndex(), bool bMergeIslands = true);
-
 	/** Adds a node to the graph's node collection and modifies the NextAvailableNodeUniqueIndex as necessary to maintain its validity. */
 	void RegisterVertex(TObjectPtr<UGraphVertex> Node);
 
@@ -173,14 +161,11 @@ private:
 	void RegisterIsland(TObjectPtr<UGraphIsland> Edge);
 
 	/** When we add multiple edges into the graph. This function will ensure that the interactions we make externally are kept to a minimum. */
-	void MergeOrCreateIslands(const TArray<FGraphEdgeHandle>& InEdges);
+	void MergeOrCreateIslands(TArray<FGraphEdgeHandle>&& InEdges);
 
 	/** After a change, this function will remove the island if it's empty or will attempt to split it into two smaller islands. */
 	void RemoveOrSplitIsland(TObjectPtr<UGraphIsland> Island);
 
-	/** Need this so that users aren't able to pass in bHandleIslands = false. */
-	void RemoveEdgeInternal(const FGraphEdgeHandle& EdgeHandle, bool bHandleIslands);
-
 	/** Factory functions for creating an appropriately typed node/edge/island. */
 	virtual TObjectPtr<UGraphVertex> CreateTypedVertex() const;
 	virtual TObjectPtr<UGraphEdge> CreateTypedEdge() const;
diff --git a/Engine/Plugins/Experimental/GameplayGraph/Source/GameplayGraph/Public/Graph/GraphEdge.h b/Engine/Plugins/Experimental/GameplayGraph/Source/GameplayGraph/Public/Graph/GraphEdge.h
--- a/Engine/Plugins/Experimental/GameplayGraph/Source/GameplayGraph/Public/Graph/GraphEdge.h
+++ b/Engine/Plugins/Experimental/GameplayGraph/Source/GameplayGraph/Public/Graph/GraphEdge.h
@@ -40,7 +40,7 @@ public:
 
 	FGraphEdgeHandle Handle() const
 	{
-		return FGraphEdgeHandle{ GetUniqueIndex(), GetGraph() };
+		return FGraphEdgeHandle{ GetUniqueIndex(), const_cast<UGraphEdge*>(this) };
 	}
 
 	FSerializedEdgeData GetSerializedData() const;
diff --git a/Engine/Plugins/Experimental/GameplayGraph/Source/GameplayGraph/Public/Graph/GraphHandle.h b/Engine/Plugins/Experimental/GameplayGraph/Source/GameplayGraph/Public/Graph/GraphHandle.h
--- a/Engine/Plugins/Experimental/GameplayGraph/Source/GameplayGraph/Public/Graph/GraphHandle.h
+++ b/Engine/Plugins/Experimental/GameplayGraph/Source/GameplayGraph/Public/Graph/GraphHandle.h
@@ -6,7 +6,7 @@
 
 #include "GraphHandle.generated.h"
 
-class UGraph;
+class UGraphElement;
 class UGraphVertex;
 class UGraphEdge;
 class UGraphIsland;
@@ -102,19 +102,20 @@ struct GAMEPLAYGRAPH_API FGraphHandle
 {
 	GENERATED_BODY()
 public:
-	FGraphHandle() = default;
-	FGraphHandle(FGraphUniqueIndex InUniqueIndex, UGraph* InGraph);
-	virtual ~FGraphHandle() = default;
+	FGraphHandle();
+	FGraphHandle(FGraphUniqueIndex InUniqueIndex, TObjectPtr<UGraphElement> InElement);
 
 	void Clear();
 
 	/** Whether or not this handle has been initialized. */
 	bool IsValid() const;
+	bool HasElement() const;
 	bool IsComplete() const;
-	virtual bool HasElement() const { return false; }
 
 	FGraphUniqueIndex GetUniqueIndex() const { return UniqueIndex; }
-	UGraph* GetGraph() const { return WeakGraph.Get(); }
+
+	void SetElement(TObjectPtr<UGraphElement> InElement);
+	TObjectPtr<UGraphElement> GetElement() const;
 
 	bool operator==(const FGraphHandle& Other) const;
 	bool operator!=(const FGraphHandle& Other) const;
@@ -128,7 +129,7 @@ private:
 
 	/** Pointer to the graph */
 	UPROPERTY(Transient)
-	TWeakObjectPtr<UGraph> WeakGraph;
+	TWeakObjectPtr<UGraphElement> Element;
 };
 
 USTRUCT()
@@ -136,15 +137,12 @@ struct GAMEPLAYGRAPH_API FGraphVertexHandle : public FGraphHandle
 {
 	GENERATED_BODY()
 
-	static FGraphVertexHandle Invalid;
-
-	FGraphVertexHandle() = default;
-	FGraphVertexHandle(FGraphUniqueIndex InUniqueIndex, UGraph* InGraph)
-		: FGraphHandle(InUniqueIndex, InGraph)
+	FGraphVertexHandle();
+	FGraphVertexHandle(FGraphUniqueIndex InUniqueIndex, TObjectPtr<UGraphElement> InElement = nullptr)
+		: FGraphHandle(InUniqueIndex, InElement)
 	{}
 
-	UGraphVertex* GetVertex() const;
-	virtual bool HasElement() const override;
+	TObjectPtr<UGraphVertex> GetVertex() const;
 };
 
 USTRUCT()
@@ -152,15 +150,12 @@ struct GAMEPLAYGRAPH_API FGraphEdgeHandle : public FGraphHandle
 {
 	GENERATED_BODY()
 
-	static FGraphEdgeHandle Invalid;
-
-	FGraphEdgeHandle() = default;
-	FGraphEdgeHandle(FGraphUniqueIndex InUniqueIndex, UGraph* InGraph)
-		: FGraphHandle(InUniqueIndex, InGraph)
+	FGraphEdgeHandle();
+	FGraphEdgeHandle(FGraphUniqueIndex InUniqueIndex, TObjectPtr<UGraphElement> InElement)
+		: FGraphHandle(InUniqueIndex, InElement)
 	{}
 
-	UGraphEdge* GetEdge() const;
-	virtual bool HasElement() const override;
+	TObjectPtr<UGraphEdge> GetEdge() const;
 };
 
 USTRUCT()
@@ -168,13 +163,10 @@ struct GAMEPLAYGRAPH_API FGraphIslandHandle : public FGraphHandle
 {
 	GENERATED_BODY()
 
-	static FGraphIslandHandle Invalid;
-
-	FGraphIslandHandle() = default;
-	FGraphIslandHandle(FGraphUniqueIndex InUniqueIndex, UGraph* InGraph)
-		: FGraphHandle(InUniqueIndex, InGraph)
+	FGraphIslandHandle();
+	FGraphIslandHandle(FGraphUniqueIndex InUniqueIndex, TObjectPtr<UGraphElement> InElement)
+		: FGraphHandle(InUniqueIndex, InElement)
 	{}
 
-	UGraphIsland* GetIsland() const;
-	virtual bool HasElement() const override;
-};
+	TObjectPtr<UGraphIsland> GetIsland() const;
+};
\ No newline at end of file
diff --git a/Engine/Plugins/Experimental/GameplayGraph/Source/GameplayGraph/Public/Graph/GraphIsland.h b/Engine/Plugins/Experimental/GameplayGraph/Source/GameplayGraph/Public/Graph/GraphIsland.h
--- a/Engine/Plugins/Experimental/GameplayGraph/Source/GameplayGraph/Public/Graph/GraphIsland.h
+++ b/Engine/Plugins/Experimental/GameplayGraph/Source/GameplayGraph/Public/Graph/GraphIsland.h
@@ -60,9 +60,10 @@ public:
 
 	FGraphIslandHandle Handle() const
 	{
-		return FGraphIslandHandle{ GetUniqueIndex(), GetGraph() };
+		return FGraphIslandHandle{ GetUniqueIndex(), const_cast<UGraphIsland*>(this) };
 	}
 
+	bool IsPendingDestroy() const { return bPendingDestroy; }
 	bool IsEmpty() const { return Vertices.IsEmpty(); }
 	const TSet<FGraphVertexHandle>& GetVertices() const { return Vertices; }
 	int32 Num() const { return Vertices.Num(); }
@@ -106,6 +107,9 @@ private:
 	UPROPERTY(SaveGame)
 	TSet<FGraphVertexHandle> Vertices;
 
+	UPROPERTY(Transient)
+	bool bPendingDestroy;
+
 	UPROPERTY(Transient)
 	EGraphIslandOperations AllowedOperations = EGraphIslandOperations::All;
 };
\ No newline at end of file
diff --git a/Engine/Plugins/Experimental/GameplayGraph/Source/GameplayGraph/Public/Graph/GraphVertex.h b/Engine/Plugins/Experimental/GameplayGraph/Source/GameplayGraph/Public/Graph/GraphVertex.h
--- a/Engine/Plugins/Experimental/GameplayGraph/Source/GameplayGraph/Public/Graph/GraphVertex.h
+++ b/Engine/Plugins/Experimental/GameplayGraph/Source/GameplayGraph/Public/Graph/GraphVertex.h
@@ -18,12 +18,11 @@ public:
 
 	FGraphVertexHandle Handle() const
 	{
-		return FGraphVertexHandle{ GetUniqueIndex(), GetGraph()};
+		return FGraphVertexHandle{ GetUniqueIndex(), const_cast<UGraphVertex*>(this)};
 	}
 
 	bool HasEdgeTo(const FGraphVertexHandle& Other) const;
 	const FGraphEdgeHandle* FindEdgeTo(const FGraphVertexHandle& Other) const;
-	int32 NumEdges() const { return Edges.Num(); }
 
 	const FGraphIslandHandle& GetParentIsland() const { return ParentIsland; }
 
@@ -42,8 +41,6 @@ public:
 			}
 		}
 	}
-	
-	const TMap<FGraphVertexHandle, FGraphEdgeHandle>& GetEdges() const { return Edges; }
 
 	FOnGraphVertexRemoved OnVertexRemoved;
 	FOnGraphVertexParentIslandSet OnParentIslandSet;
diff --git a/Engine/Plugins/Experimental/GameplayGraph/Tests/GameplayGraphTests.Build.cs b/Engine/Plugins/Experimental/GameplayGraph/Tests/GameplayGraphTests.Build.cs
deleted file mode 100644
diff --git a/Engine/Plugins/Experimental/GameplayGraph/Tests/GameplayGraphTests.Target.cs b/Engine/Plugins/Experimental/GameplayGraph/Tests/GameplayGraphTests.Target.cs
deleted file mode 100644
diff --git a/Engine/Plugins/Experimental/GameplayGraph/Tests/Private/GraphEdgeTests.cpp b/Engine/Plugins/Experimental/GameplayGraph/Tests/Private/GraphEdgeTests.cpp
deleted file mode 100644
diff --git a/Engine/Plugins/Experimental/GameplayGraph/Tests/Private/GraphElementTests.cpp b/Engine/Plugins/Experimental/GameplayGraph/Tests/Private/GraphElementTests.cpp
deleted file mode 100644
diff --git a/Engine/Plugins/Experimental/GameplayGraph/Tests/Private/GraphHandleTests.cpp b/Engine/Plugins/Experimental/GameplayGraph/Tests/Private/GraphHandleTests.cpp
deleted file mode 100644
diff --git a/Engine/Plugins/Experimental/GameplayGraph/Tests/Private/GraphIslandTests.cpp b/Engine/Plugins/Experimental/GameplayGraph/Tests/Private/GraphIslandTests.cpp
deleted file mode 100644
diff --git a/Engine/Plugins/Experimental/GameplayGraph/Tests/Private/GraphSerializationTests.cpp b/Engine/Plugins/Experimental/GameplayGraph/Tests/Private/GraphSerializationTests.cpp
deleted file mode 100644
diff --git a/Engine/Plugins/Experimental/GameplayGraph/Tests/Private/GraphTests.cpp b/Engine/Plugins/Experimental/GameplayGraph/Tests/Private/GraphTests.cpp
deleted file mode 100644
diff --git a/Engine/Plugins/Experimental/GameplayGraph/Tests/Private/GraphVertexTests.cpp b/Engine/Plugins/Experimental/GameplayGraph/Tests/Private/GraphVertexTests.cpp
deleted file mode 100644
diff --git a/Engine/Plugins/Experimental/GameplayGraph/Tests/Private/TestGraphBuilder.cpp b/Engine/Plugins/Experimental/GameplayGraph/Tests/Private/TestGraphBuilder.cpp
deleted file mode 100644
diff --git a/Engine/Plugins/Experimental/GameplayGraph/Tests/Private/TestGraphBuilder.h b/Engine/Plugins/Experimental/GameplayGraph/Tests/Private/TestGraphBuilder.h
deleted file mode 100644
diff --git a/Engine/Plugins/Experimental/GameplayGraph/Tests/Private/TestGraphElement.cpp b/Engine/Plugins/Experimental/GameplayGraph/Tests/Private/TestGraphElement.cpp
deleted file mode 100644
diff --git a/Engine/Plugins/Experimental/GameplayGraph/Tests/Private/TestGraphElement.h b/Engine/Plugins/Experimental/GameplayGraph/Tests/Private/TestGraphElement.h
deleted file mode 100644
diff --git a/Engine/Plugins/Experimental/GameplayGraph/Tests/Private/TestGroupEvents.cpp b/Engine/Plugins/Experimental/GameplayGraph/Tests/Private/TestGroupEvents.cpp
deleted file mode 100644
