more flexible graph serialization and some gameplay graph sanity checks

- Removed the old UGraph::GetSerializableGraph and UGraph::LoadFromSerializedGraph functions and replaced with IGraphSerialization and IGraphDeserialization to provide more flexibility in how the graph is serialized/deserialized
- UGraph::RemoveBulkVertices and UGraphIsland::RemoveVertex should remove vertex even if the incoming handle is incomplete but the GUID refers to a valid vertex
- UGraph::RemoveIsland should remove island even if the incoming handle is incomplete but the GUID refers to a valid island

#tests low level tests, manual testing

diff --git a/Engine/Plugins/Experimental/GameplayGraph/Source/GameplayGraph/Private/Graph/Graph.cpp b/Engine/Plugins/Experimental/GameplayGraph/Source/GameplayGraph/Private/Graph/Graph.cpp
--- a/Engine/Plugins/Experimental/GameplayGraph/Source/GameplayGraph/Private/Graph/Graph.cpp
+++ b/Engine/Plugins/Experimental/GameplayGraph/Source/GameplayGraph/Private/Graph/Graph.cpp
@@ -1,9 +1,10 @@
 // Copyright Epic Games, Inc. All Rights Reserved.
-
 #include "Graph/Graph.h"
+
+#include "Graph/Algorithms/Connectivity/ConnectedComponents.h"
 #include "Graph/GraphEdge.h"
 #include "Graph/GraphVertex.h"
-#include "Graph/Algorithms/Connectivity/ConnectedComponents.h"
+#include "Graph/GraphSerialization.h"
 #include "GenericPlatform/GenericPlatformMath.h"
 #include "HAL/IConsoleManager.h"
 
@@ -12,116 +13,10 @@
 
 namespace Graph
 {
-	static bool bFixupGraphOnLoad = false;
+	static bool bFixupGraphOnLoad = true;
 	static FAutoConsoleVariableRef FixupGraphOnLoadCVar(TEXT("GameplayGraph.FixupGraphOnLoad"), bFixupGraphOnLoad, TEXT("Merge and split islands when loaded from serialization."));
 }
 
-FSerializableGraph UGraph::GetSerializableGraph() const
-{
-	FSerializableGraph Out;
-	Out.Properties = Properties;
-	Out.Vertices.Reserve(Vertices.Num());
-	Out.Edges.Reserve(Edges.Num());
-	Out.Islands.Reserve(Islands.Num());
-
-	for (const TPair<FGraphVertexHandle, TObjectPtr<UGraphVertex>>& Kvp : Vertices)
-	{
-		if (!Kvp.Key.IsComplete() || !Kvp.Value)
-		{
-			continue;
-		}
-		Out.Vertices.Add(Kvp.Key);
-	}
-
-	for (const TPair<FGraphEdgeHandle, TObjectPtr<UGraphEdge>>& Kvp : Edges)
-	{
-		if (!Kvp.Key.IsComplete() || !Kvp.Value)
-		{
-			continue;
-		}
-
-		Out.Edges.Add(Kvp.Key, Kvp.Value->GetSerializedData());
-	}
-
-	for (const TPair<FGraphIslandHandle, TObjectPtr<UGraphIsland>>& Kvp : Islands)
-	{
-		if (!Kvp.Key.IsComplete() || !Kvp.Value)
-		{
-			continue;
-		}
-
-		Out.Islands.Add(Kvp.Key, Kvp.Value->GetSerializedData());
-	}
-
-	return Out;
-}
-
-void UGraph::LoadFromSerializedGraph(const FSerializableGraph& Input)
-{
-	InitializeFromProperties(Input.Properties);
-
-	for (const FGraphVertexHandle& Handle : Input.Vertices)
-	{
-		if (!Handle.IsValid())
-		{
-			continue;
-		}
-
-		CreateVertex(Handle.GetUniqueIndex());
-	}
-
-	TArray<FGraphEdgeHandle> AllEdges;
-	AllEdges.Reserve(Input.Edges.Num());
-	for (const TPair<FGraphEdgeHandle, FSerializedEdgeData>& Kvp : Input.Edges)
-	{
-		if (!Kvp.Key.IsValid())
-		{
-			continue;
-		}
-
-		AllEdges.Emplace(CreateEdge(
-			GetCompleteNodeHandle(Kvp.Value.Node1),
-			GetCompleteNodeHandle(Kvp.Value.Node2),
-			Kvp.Key.GetUniqueIndex(),
-			false
-		));
-	}
-
-	TArray<FGraphIslandHandle> AllIslands;
-	AllIslands.Reserve(Input.Islands.Num());
-	for (const TPair<FGraphIslandHandle, FSerializedIslandData>& Kvp : Input.Islands)
-	{
-		if (!Kvp.Key.IsValid())
-		{
-			continue;
-		}
-
-		TArray<FGraphVertexHandle> IslandVertices;
-		IslandVertices.Reserve(Kvp.Value.Vertices.Num());
-		
-		for (const FGraphVertexHandle& VertexHandle : Kvp.Value.Vertices)
-		{
-			FGraphVertexHandle CompleteHandle = GetCompleteNodeHandle(VertexHandle);
-			if (CompleteHandle.IsComplete())
-			{
-				IslandVertices.Add(CompleteHandle);
-			}
-		}
-
-		AllIslands.Emplace(CreateIsland(IslandVertices, Kvp.Key.GetUniqueIndex()));
-	}
-
-	if (Graph::bFixupGraphOnLoad)
-	{
-		MergeOrCreateIslands(AllEdges);
-
-		for (const FGraphIslandHandle& NewIslandHandle : AllIslands)
-		{
-			RemoveOrSplitIsland(NewIslandHandle.GetIsland());
-		}
-	}
-}
-
 void UGraph::Empty()
 {
 	Vertices.Empty();
@@ -365,7 +260,7 @@ void UGraph::RemoveIsland(const FGraphIslandHandle& IslandHandle)
 		return;
 	}
 
-	if (TObjectPtr<UGraphIsland> Island = IslandHandle.GetIsland())
+	if (UGraphIsland* Island = GetSafeIslandFromHandle(IslandHandle))
 	{
 		if (!Island->IsOperationAllowed(EGraphIslandOperations::Destroy))
 		{
@@ -602,39 +497,40 @@ void UGraph::RemoveBulkVertices(const TArray<FGraphVertexHandle>& InHandles)
 
 	TSet<FGraphIslandHandle> AffectedIslands;
 
+	TArray<FGraphVertexHandle> AffectedVertices;
+	AffectedVertices.Reserve(InHandles.Num());
+
 	for (const FGraphVertexHandle& NodeHandle : InHandles)
 	{
-		if (NodeHandle.IsValid())
+		if (UGraphVertex* Node = GetSafeVertexFromHandle(NodeHandle))
 		{
-			if (TObjectPtr<UGraphVertex> Node = NodeHandle.GetVertex())
+			if (TObjectPtr<UGraphIsland> Island = Node->GetParentIsland().GetIsland())
 			{
-				if (TObjectPtr<UGraphIsland> Island = Node->GetParentIsland().GetIsland())
-				{
-					AffectedIslands.Add(Node->GetParentIsland());
-					Island->RemoveVertex(NodeHandle);
-				}
+				AffectedIslands.Add(Node->GetParentIsland());
+				Island->RemoveVertex(NodeHandle);
+			}
 
-				// We must remove every edge this node is a part of. Need to make a copy of the edges because
-				// otherwise we're modifying the container during iteration over it in ForEachAdjacentVertex.
-				TArray<FGraphEdgeHandle> EdgeCopy;
-				EdgeCopy.Reserve(Node->NumEdges());
-				Node->ForEachAdjacentVertex(
-					[&EdgeCopy](const FGraphVertexHandle& OtherNodeHandle, const FGraphEdgeHandle& EdgeHandle)
-					{
-						EdgeCopy.Add(EdgeHandle);
-					}
-				);
-
-				for (const FGraphEdgeHandle& EdgeHandle : EdgeCopy)
+			// We must remove every edge this node is a part of. Need to make a copy of the edges because
+			// otherwise we're modifying the container during iteration over it in ForEachAdjacentVertex.
+			TArray<FGraphEdgeHandle> EdgeCopy;
+			EdgeCopy.Reserve(Node->NumEdges());
+			Node->ForEachAdjacentVertex(
+				[&EdgeCopy](const FGraphVertexHandle& OtherNodeHandle, const FGraphEdgeHandle& EdgeHandle)
 				{
-					// Don't immediately handle islands. We'll do it later.
-					RemoveEdgeInternal(EdgeHandle, false);
+					EdgeCopy.Add(EdgeHandle);
 				}
+			);
+
+			for (const FGraphEdgeHandle& EdgeHandle : EdgeCopy)
+			{
+				// Don't immediately handle islands. We'll do it later.
+				RemoveEdgeInternal(EdgeHandle, false);
 			}
 
+			AffectedVertices.Add(NodeHandle);
 		}
 	}
-	
+
 	if (Properties.bGenerateIslands)
 	{
 		for (const FGraphIslandHandle& IslandHandle : AffectedIslands)
@@ -645,16 +541,13 @@ void UGraph::RemoveBulkVertices(const TArray<FGraphVertexHandle>& InHandles)
 
 	// A final pass after generation of islands to clean up book-keeping.
 	// TODO: Not sure if this is necessary and can be done before we regenerate islands?
-	for (const FGraphVertexHandle& NodeHandle : InHandles)
+	for (const FGraphVertexHandle& NodeHandle : AffectedVertices)
 	{
-		if (NodeHandle.IsValid())
+		if (UGraphVertex* Node = NodeHandle.GetVertex())
 		{
-			if (TObjectPtr<UGraphVertex> Node = NodeHandle.GetVertex())
-			{
-				Node->HandleOnVertexRemoved();
-			}
-			Vertices.Remove(NodeHandle);
+			Node->HandleOnVertexRemoved();
 		}
+		Vertices.Remove(NodeHandle);
 	}
 }
 
@@ -859,4 +752,147 @@ void UGraph::RefreshIslandConnectivity(const FGraphIslandHandle& IslandHandle)
 {
 	TRACE_CPUPROFILER_EVENT_SCOPE(UGraph::RefreshIslandConnectivity);
 	RemoveOrSplitIsland(IslandHandle.GetIsland());
+}
+
+UGraphVertex* UGraph::GetSafeVertexFromHandle(const FGraphVertexHandle& Handle) const
+{
+	return Handle.IsComplete() ? Handle.GetVertex() : Vertices.FindRef(Handle).Get();
+}
+
+UGraphEdge* UGraph::GetSafeEdgeFromHandle(const FGraphEdgeHandle& Handle) const
+{
+	return Handle.IsComplete() ? Handle.GetEdge() : Edges.FindRef(Handle).Get();
+}
+
+UGraphIsland* UGraph::GetSafeIslandFromHandle(const FGraphIslandHandle& Handle) const
+{
+	return Handle.IsComplete() ? Handle.GetIsland() : Islands.FindRef(Handle).Get();
+}
+
+void UGraph::ReserveVertices(int32 Delta)
+{
+	Vertices.Reserve(Vertices.Num() + Delta);
+}
+
+void UGraph::ReserveEdges(int32 Delta)
+{
+	Edges.Reserve(Edges.Num() + Delta);
+}
+
+void UGraph::ReserveIslands(int32 Delta)
+{
+	Islands.Reserve(Islands.Num() + Delta);
+}
+
+void operator<<(IGraphSerialization& Output, const UGraph& Graph)
+{
+	Output.WriteGraphProperties(Graph.GetProperties());
+	for (const TPair<FGraphVertexHandle, TObjectPtr<UGraphVertex>>& Kvp : Graph.GetVertices())
+	{
+		if (!ensure(Kvp.Key.IsComplete() && Kvp.Value))
+		{
+			continue;
+		}
+
+		Output.WriteGraphVertex(Kvp.Key, Kvp.Value);
+	}
+
+	for (const TPair<FGraphEdgeHandle, TObjectPtr<UGraphEdge>>& Kvp : Graph.GetEdges())
+	{
+		if (!ensure(Kvp.Key.IsComplete() && Kvp.Value))
+		{
+			continue;
+		}
+
+		Output.WriteGraphEdge(Kvp.Key, Kvp.Value);
+	}
+
+	for (const TPair<FGraphIslandHandle, TObjectPtr<UGraphIsland>>& Kvp : Graph.GetIslands())
+	{
+		if (!ensure(Kvp.Key.IsComplete() && Kvp.Value))
+		{
+			continue;
+		}
+
+		Output.WriteGraphIsland(Kvp.Key, Kvp.Value);
+	}
+}
+
+void operator>>(const IGraphDeserialization& Input, UGraph& Graph)
+{
+	Graph.InitializeFromProperties(Input.GetProperties());
+
+	Graph.ReserveVertices(Input.NumVertices());
+	Input.ForEveryVertex(
+		[&Graph](const FGraphVertexHandle& InHandle)
+		{
+			if (!ensure(InHandle.IsValid()))
+			{
+				return FGraphVertexHandle{};
+			}
+
+			return Graph.CreateVertex(InHandle.GetUniqueIndex());
+		}
+	);
+
+	TArray<FGraphEdgeHandle> AllEdges;
+	AllEdges.Reserve(Input.NumEdges());
+	Graph.ReserveEdges(Input.NumEdges());
+	Input.ForEveryEdge(
+		[&Graph, &AllEdges](const FGraphEdgeHandle& InHandle, const IGraphDeserialization::FEdgeConstructionData& Data)
+		{
+			if (!ensure(InHandle.IsValid()))
+			{
+				return FGraphEdgeHandle{};
+			}
+
+			FGraphEdgeHandle Handle = Graph.CreateEdge(
+				Graph.GetCompleteNodeHandle(Data.Vertex1),
+				Graph.GetCompleteNodeHandle(Data.Vertex2),
+				InHandle.GetUniqueIndex(),
+				false
+			);
+			AllEdges.Add(Handle);
+			return Handle;
+		}
+	);
+
+	TArray<FGraphIslandHandle> AllIslands;
+	AllIslands.Reserve(Input.NumIslands());
+	Graph.ReserveIslands(Input.NumIslands());
+	Input.ForEveryIsland(
+		[&Graph, &AllIslands](const FGraphIslandHandle& InHandle, const IGraphDeserialization::FIslandConstructionData& Data)
+		{
+			if (!ensure(InHandle.IsValid()))
+			{
+				return FGraphIslandHandle{};
+			}
+
+			TArray<FGraphVertexHandle> IslandVertices;
+			IslandVertices.Reserve(Data.Vertices.Num());
+		
+			for (const FGraphVertexHandle& VertexHandle : Data.Vertices)
+			{
+				FGraphVertexHandle CompleteHandle = Graph.GetCompleteNodeHandle(VertexHandle);
+				if (CompleteHandle.IsComplete())
+				{
+					IslandVertices.Add(CompleteHandle);
+				}
+			}
+
+			FGraphIslandHandle NewIslandHandle = Graph.CreateIsland(IslandVertices, InHandle.GetUniqueIndex());
+			AllIslands.Add(NewIslandHandle);
+			return NewIslandHandle;
+		}
+	);
+
+	if (Graph::bFixupGraphOnLoad)
+	{
+		Graph.MergeOrCreateIslands(AllEdges);
+
+		for (const FGraphIslandHandle& NewIslandHandle : AllIslands)
+		{
+			Graph.RemoveOrSplitIsland(NewIslandHandle.GetIsland());
+		}
+	}
 }
\ No newline at end of file
diff --git a/Engine/Plugins/Experimental/GameplayGraph/Source/GameplayGraph/Private/Graph/GraphDefaultSerialization.cpp b/Engine/Plugins/Experimental/GameplayGraph/Source/GameplayGraph/Private/Graph/GraphDefaultSerialization.cpp
new file mode 100644
--- /dev/null
+++ b/Engine/Plugins/Experimental/GameplayGraph/Source/GameplayGraph/Private/Graph/GraphDefaultSerialization.cpp
@@ -0,0 +1,5 @@
+// Copyright Epic Games, Inc. All Rights Reserved.
+
+#include "Graph/GraphDefaultSerialization.h"
+
+#include UE_INLINE_GENERATED_CPP_BY_NAME(GraphDefaultSerialization)
\ No newline at end of file
diff --git a/Engine/Plugins/Experimental/GameplayGraph/Source/GameplayGraph/Private/Graph/GraphEdge.cpp b/Engine/Plugins/Experimental/GameplayGraph/Source/GameplayGraph/Private/Graph/GraphEdge.cpp
--- a/Engine/Plugins/Experimental/GameplayGraph/Source/GameplayGraph/Private/Graph/GraphEdge.cpp
+++ b/Engine/Plugins/Experimental/GameplayGraph/Source/GameplayGraph/Private/Graph/GraphEdge.cpp
@@ -35,11 +35,3 @@ bool UGraphEdge::ContainsNode(const FGraphVertexHandle& InNode) const
 {
 	return (InNode == A) || (InNode == B);
 }
-
-FSerializedEdgeData UGraphEdge::GetSerializedData() const
-{
-	FSerializedEdgeData Nodes;
-	Nodes.Node1 = A;
-	Nodes.Node2 = B;
-	return Nodes;
-}
\ No newline at end of file
diff --git a/Engine/Plugins/Experimental/GameplayGraph/Source/GameplayGraph/Private/Graph/GraphIsland.cpp b/Engine/Plugins/Experimental/GameplayGraph/Source/GameplayGraph/Private/Graph/GraphIsland.cpp
--- a/Engine/Plugins/Experimental/GameplayGraph/Source/GameplayGraph/Private/Graph/GraphIsland.cpp
+++ b/Engine/Plugins/Experimental/GameplayGraph/Source/GameplayGraph/Private/Graph/GraphIsland.cpp
@@ -68,26 +68,14 @@ void UGraphIsland::RemoveVertex(const FGraphVertexHandle& Node)
 		return;
 	}
 
-	TObjectPtr<UGraphVertex> NodePtr = Node.GetVertex();
-	if (!NodePtr)
+	if (UGraphVertex* NodePtr = Node.GetVertex())
 	{
-		return;
+		NodePtr->SetParentIsland({});
 	}
-	NodePtr->SetParentIsland({});
 	Vertices.Remove(Node);
 	HandleOnVertexRemoved(Node);
 }
 
-FSerializedIslandData UGraphIsland::GetSerializedData() const
-{
-	FSerializedIslandData Out;
-	for (const FGraphVertexHandle& Handle : Vertices)
-	{
-		Out.Vertices.Add(Handle);
-	}
-	return Out;
-}
-
 void UGraphIsland::HandleOnVertexAdded(const FGraphVertexHandle& Handle)
 {
 	OnVertexAdded.Broadcast(this->Handle(), Handle);
diff --git a/Engine/Plugins/Experimental/GameplayGraph/Source/GameplayGraph/Private/Graph/GraphSerialization.cpp b/Engine/Plugins/Experimental/GameplayGraph/Source/GameplayGraph/Private/Graph/GraphSerialization.cpp
new file mode 100644
--- /dev/null
+++ b/Engine/Plugins/Experimental/GameplayGraph/Source/GameplayGraph/Private/Graph/GraphSerialization.cpp
@@ -0,0 +1,3 @@
+// Copyright Epic Games, Inc. All Rights Reserved.
+
+#include "Graph/GraphSerialization.h"
diff --git a/Engine/Plugins/Experimental/GameplayGraph/Source/GameplayGraph/Public/Graph/Graph.h b/Engine/Plugins/Experimental/GameplayGraph/Source/GameplayGraph/Public/Graph/Graph.h
--- a/Engine/Plugins/Experimental/GameplayGraph/Source/GameplayGraph/Public/Graph/Graph.h
+++ b/Engine/Plugins/Experimental/GameplayGraph/Source/GameplayGraph/Public/Graph/Graph.h
@@ -14,6 +14,9 @@
 
 #include "Graph.generated.h"
 
+class IGraphSerialization;
+class IGraphDeserialization;
+
 USTRUCT()
 struct FGraphProperties
 {
@@ -31,31 +34,6 @@ DECLARE_MULTICAST_DELEGATE_OneParam(FOnGraphVertexCreated, const FGraphVertexHan
 DECLARE_MULTICAST_DELEGATE_OneParam(FOnGraphEdgeCreated, const FGraphEdgeHandle&);
 DECLARE_MULTICAST_DELEGATE_OneParam(FOnGraphIslandCreated, const FGraphIslandHandle&);
 
-/**
- * The minimum amount of data we need to serialize to be able to reconstruct the graph as it was.
- * Note that classes that inherit from UGraph and its elements will no doubt want to extend the graph
- * with actual information on each node/edge/island. In that case, they should extend FSerializableGraph
- * to contain the extra information per graph handle. Furthermore, they'll need to extend UGraph to have
- * its own typed serialization save/load functions that call the base functions in UGraph first.
- */
-USTRUCT()
-struct FSerializableGraph
-{
-	GENERATED_BODY()
-
-	UPROPERTY(SaveGame)
-	FGraphProperties Properties;
-
-	UPROPERTY(SaveGame)
-	TArray<FGraphVertexHandle> Vertices;
-
-	UPROPERTY(SaveGame)
-	TMap<FGraphEdgeHandle, FSerializedEdgeData> Edges;
-
-	UPROPERTY(SaveGame)
-	TMap<FGraphIslandHandle, FSerializedIslandData> Islands;
-};
-
 struct FEdgeCreationParameters
 {
 	FGraphVertexHandle VertexHandle1;
@@ -92,9 +70,6 @@ class GAMEPLAYGRAPH_API UGraph: public UObject
 public:
 	UGraph() = default;
 
-	FSerializableGraph GetSerializableGraph() const;
-	void LoadFromSerializedGraph(const FSerializableGraph& Input);
-
 	void Empty();
 	void InitializeFromProperties(const FGraphProperties& Properties);
 
@@ -146,6 +121,24 @@ public:
 	const TMap<FGraphIslandHandle, TObjectPtr<UGraphIsland>>& GetIslands() const { return Islands; }
 	const FGraphProperties& GetProperties() const { return Properties; }
 
+	UGraphVertex* GetSafeVertexFromHandle(const FGraphVertexHandle& Handle) const;
+	UGraphEdge* GetSafeEdgeFromHandle(const FGraphEdgeHandle& Handle) const;
+	UGraphIsland* GetSafeIslandFromHandle(const FGraphIslandHandle& Handle) const;
+
+	GAMEPLAYGRAPH_API friend void operator<<(IGraphSerialization& Output, const UGraph& Graph);
+	GAMEPLAYGRAPH_API friend void operator>>(const IGraphDeserialization& Input, UGraph& Graph);
+
+protected:
+	/** When we add multiple edges into the graph. This function will ensure that the interactions we make externally are kept to a minimum. */
+	void MergeOrCreateIslands(const TArray<FGraphEdgeHandle>& InEdges);
+
+	/** After a change, this function will remove the island if it's empty or will attempt to split it into two smaller islands. */
+	void RemoveOrSplitIsland(TObjectPtr<UGraphIsland> Island);
+
+	/** Used for bulk loading in vertices/edges/islands. */
+	void ReserveVertices(int32 Delta);
+	void ReserveEdges(int32 Delta);
+	void ReserveIslands(int32 Delta);
 private:
 	UPROPERTY()
 	TMap<FGraphVertexHandle, TObjectPtr<UGraphVertex>> Vertices;
@@ -172,12 +165,6 @@ private:
 	/** Add an island to the graph and modifies the NextAvailableIslandUniqueIndex as necessary to maintain its validity. */
 	void RegisterIsland(TObjectPtr<UGraphIsland> Edge);
 
-	/** When we add multiple edges into the graph. This function will ensure that the interactions we make externally are kept to a minimum. */
-	void MergeOrCreateIslands(const TArray<FGraphEdgeHandle>& InEdges);
-
-	/** After a change, this function will remove the island if it's empty or will attempt to split it into two smaller islands. */
-	void RemoveOrSplitIsland(TObjectPtr<UGraphIsland> Island);
-
 	/** Need this so that users aren't able to pass in bHandleIslands = false. */
 	void RemoveEdgeInternal(const FGraphEdgeHandle& EdgeHandle, bool bHandleIslands);
 
diff --git a/Engine/Plugins/Experimental/GameplayGraph/Source/GameplayGraph/Public/Graph/GraphDefaultSerialization.h b/Engine/Plugins/Experimental/GameplayGraph/Source/GameplayGraph/Public/Graph/GraphDefaultSerialization.h
new file mode 100644
--- /dev/null
+++ b/Engine/Plugins/Experimental/GameplayGraph/Source/GameplayGraph/Public/Graph/GraphDefaultSerialization.h
@@ -0,0 +1,209 @@
+// Copyright Epic Games, Inc. All Rights Reserved.
+#pragma once
+
+#include "Graph/Graph.h"
+#include "Graph/GraphSerialization.h"
+
+#include <concepts>
+#include <type_traits>
+
+#include "GraphDefaultSerialization.generated.h"
+
+USTRUCT()
+struct FSerializedEdgeData
+{
+	GENERATED_BODY()
+
+	UPROPERTY(SaveGame)
+	FGraphVertexHandle Node1;
+
+	UPROPERTY(SaveGame)
+	FGraphVertexHandle Node2;
+
+	// Comparison operators
+	friend bool operator==(const FSerializedEdgeData& Lhs, const FSerializedEdgeData& Rhs) = default;
+	friend bool operator!=(const FSerializedEdgeData& Lhs, const FSerializedEdgeData& Rhs) = default;
+};
+
+USTRUCT()
+struct FSerializedIslandData
+{
+	GENERATED_BODY()
+
+	UPROPERTY(SaveGame)
+	TArray<FGraphVertexHandle> Vertices;
+
+	// Comparison operators
+	friend bool operator==(const FSerializedIslandData& Lhs, const FSerializedIslandData& Rhs) = default;
+	friend bool operator!=(const FSerializedIslandData& Lhs, const FSerializedIslandData& Rhs) = default;
+};
+
+/**
+ * The minimum amount of data we need to serialize to be able to reconstruct the graph as it was.
+ * Note that classes that inherit from UGraph and its elements will no doubt want to extend the graph
+ * with actual information on each node/edge/island. In that case, they should extend FSerializableGraph
+ * to contain the extra information per graph handle. Furthermore, they'll need to extend UGraph to have
+ * its own typed serialization save/load functions that call the base functions in UGraph first.
+ */
+USTRUCT()
+struct FSerializableGraph
+{
+	GENERATED_BODY()
+
+	UPROPERTY(SaveGame)
+	FGraphProperties Properties;
+
+	UPROPERTY(SaveGame)
+	TArray<FGraphVertexHandle> Vertices;
+
+	UPROPERTY(SaveGame)
+	TMap<FGraphEdgeHandle, FSerializedEdgeData> Edges;
+
+	UPROPERTY(SaveGame)
+	TMap<FGraphIslandHandle, FSerializedIslandData> Islands;
+};
+
+template<typename T>
+concept CDefaultGraphSerializable = std::is_base_of_v<FSerializableGraph, T>;
+
+template<CDefaultGraphSerializable TSerializableGraph>
+class TDefaultGraphSerialization : public IGraphSerialization
+{
+public:
+	virtual ~TDefaultGraphSerialization() = default;
+
+	const TSerializableGraph& GetData() const { return Data; }
+
+	virtual void Initialize(int32 NumVertices, int32 NumEdges, int32 NumIslands) override
+	{
+		Data.Vertices.Reserve(NumVertices);
+		Data.Edges.Reserve(NumEdges);
+		Data.Islands.Reserve(NumIslands);
+	}
+
+	virtual void WriteGraphProperties(const FGraphProperties& Properties) override
+	{
+		Data.Properties = Properties;
+	}
+
+	virtual void WriteGraphVertex(const FGraphVertexHandle& VertexHandle, const UGraphVertex* Vertex) override
+	{
+		Data.Vertices.Add(VertexHandle);
+	}
+
+	virtual void WriteGraphEdge(const FGraphEdgeHandle& EdgeHandle, const UGraphEdge* Edge) override
+	{
+		if (!ensure(Edge))
+		{
+			return;
+		}
+
+		FSerializedEdgeData Serialized;
+		Serialized.Node1 = Edge->NodeA();
+		Serialized.Node2 = Edge->NodeB();
+		Data.Edges.Emplace(EdgeHandle, MoveTemp(Serialized));
+	}
+
+	virtual void WriteGraphIsland(const FGraphIslandHandle& IslandHandle, const UGraphIsland* Island) override
+	{
+		if (!ensure(Island))
+		{
+			return;
+		}
+
+		FSerializedIslandData Serialized;
+		for (const FGraphVertexHandle& Handle : Island->GetVertices())
+		{
+			Serialized.Vertices.Add(Handle);
+		}
+
+		Data.Islands.Emplace(IslandHandle, MoveTemp(Serialized));
+	}
+
+protected:
+	TSerializableGraph Data;
+};
+
+template<CDefaultGraphSerializable TSerializableGraph>
+class TDefaultGraphDeserialization : public IGraphDeserialization
+{
+public:
+	explicit TDefaultGraphDeserialization(const TSerializableGraph& InData)
+	: Data(InData)
+	{
+	}
+	virtual ~TDefaultGraphDeserialization() = default;
+
+	virtual const FGraphProperties& GetProperties() const override
+	{
+		return Data.Properties;
+	}
+
+	virtual int32 NumVertices() const override
+	{
+		return Data.Vertices.Num();
+	}
+
+	virtual void ForEveryVertex(const TFunction<FGraphVertexHandle(const FGraphVertexHandle&)>& Lambda) const override
+	{
+		for (const FGraphVertexHandle& SerializedHandle : Data.Vertices)
+		{
+			const FGraphVertexHandle FinalHandle = Lambda(SerializedHandle);
+			if (ensure(FinalHandle.IsValid()))
+			{
+				OnDeserializedVertex(FinalHandle);
+			}
+		}
+	}
+
+	virtual int32 NumEdges() const override
+	{
+		return Data.Edges.Num();
+	}
+
+	virtual void ForEveryEdge(const TFunction<FGraphEdgeHandle(const FGraphEdgeHandle&, const FEdgeConstructionData&)>& Lambda) const override
+	{
+		for (const TPair<FGraphEdgeHandle, FSerializedEdgeData>& Serialized : Data.Edges)
+		{
+			FEdgeConstructionData Construction;
+			Construction.Vertex1 = Serialized.Value.Node1;
+			Construction.Vertex2 = Serialized.Value.Node2;
+
+			const FGraphEdgeHandle FinalHandle = Lambda(Serialized.Key, Construction);
+			if (ensure(FinalHandle.IsValid()))
+			{
+				OnDeserializedEdge(FinalHandle);
+			}
+		}
+	}
+
+	virtual int32 NumIslands() const override
+	{
+		return Data.Islands.Num();
+	}
+
+	virtual void ForEveryIsland(const TFunction<FGraphIslandHandle(const FGraphIslandHandle&, const FIslandConstructionData&)>& Lambda) const override
+	{
+		for (const TPair<FGraphIslandHandle, FSerializedIslandData>& Serialized : Data.Islands)
+		{
+			FIslandConstructionData Construction;
+			Construction.Vertices = Serialized.Value.Vertices;
+
+			const FGraphIslandHandle FinalHandle = Lambda(Serialized.Key, Construction);
+			if (ensure(FinalHandle.IsValid()))
+			{
+				OnDeserializedIsland(FinalHandle);
+			}
+		}
+	}
+
+protected:
+	virtual void OnDeserializedVertex(const FGraphVertexHandle& VertexHandle) const {}
+	virtual void OnDeserializedEdge(const FGraphEdgeHandle& EdgeHandle) const {}
+	virtual void OnDeserializedIsland(const FGraphIslandHandle& IslandHandle) const {}
+
+	const TSerializableGraph& Data;
+};
+
+using FDefaultGraphSerialization = TDefaultGraphSerialization<FSerializableGraph>;
+using FDefaultGraphDeserialization = TDefaultGraphDeserialization<FSerializableGraph>;
\ No newline at end of file
diff --git a/Engine/Plugins/Experimental/GameplayGraph/Source/GameplayGraph/Public/Graph/GraphEdge.h b/Engine/Plugins/Experimental/GameplayGraph/Source/GameplayGraph/Public/Graph/GraphEdge.h
--- a/Engine/Plugins/Experimental/GameplayGraph/Source/GameplayGraph/Public/Graph/GraphEdge.h
+++ b/Engine/Plugins/Experimental/GameplayGraph/Source/GameplayGraph/Public/Graph/GraphEdge.h
@@ -6,22 +6,6 @@
 
 class UGraphVertex;
 
-USTRUCT()
-struct FSerializedEdgeData
-{
-	GENERATED_BODY()
-
-	UPROPERTY(SaveGame)
-	FGraphVertexHandle Node1;
-
-	UPROPERTY(SaveGame)
-	FGraphVertexHandle Node2;
-
-	// Comparison operators
-	friend bool operator==(const FSerializedEdgeData& Lhs, const FSerializedEdgeData& Rhs) = default;
-	friend bool operator!=(const FSerializedEdgeData& Lhs, const FSerializedEdgeData& Rhs) = default;
-};
-
 UCLASS()
 class GAMEPLAYGRAPH_API UGraphEdge : public UGraphElement
 {
@@ -43,8 +27,6 @@ public:
 		return FGraphEdgeHandle{ GetUniqueIndex(), GetGraph() };
 	}
 
-	FSerializedEdgeData GetSerializedData() const;
-
 	friend class UGraph;
 
 protected:
diff --git a/Engine/Plugins/Experimental/GameplayGraph/Source/GameplayGraph/Public/Graph/GraphIsland.h b/Engine/Plugins/Experimental/GameplayGraph/Source/GameplayGraph/Public/Graph/GraphIsland.h
--- a/Engine/Plugins/Experimental/GameplayGraph/Source/GameplayGraph/Public/Graph/GraphIsland.h
+++ b/Engine/Plugins/Experimental/GameplayGraph/Source/GameplayGraph/Public/Graph/GraphIsland.h
@@ -5,19 +5,6 @@
 #include "Graph/GraphElement.h"
 #include "GraphIsland.generated.h"
 
-USTRUCT()
-struct FSerializedIslandData
-{
-	GENERATED_BODY()
-
-	UPROPERTY(SaveGame)
-	TArray<FGraphVertexHandle> Vertices;
-
-	// Comparison operators
-	friend bool operator==(const FSerializedIslandData& Lhs, const FSerializedIslandData& Rhs) = default;
-	friend bool operator!=(const FSerializedIslandData& Lhs, const FSerializedIslandData& Rhs) = default;
-};
-
 /**
  * These are the possible operations that can be done to an island. The graph
  * will attempt to check that the island is allowing these operations before
@@ -28,7 +15,7 @@ enum class EGraphIslandOperations : int32
 {
 	None = 0,
 	Add = 1 << 0,
-	Split = 1 << 1,
+Split = 1 << 1,
 	Merge = 1 << 2,
 	Destroy = 1 << 3,
 	All = Add | Split | Merge | Destroy
@@ -66,7 +53,6 @@ public:
 	bool IsEmpty() const { return Vertices.IsEmpty(); }
 	const TSet<FGraphVertexHandle>& GetVertices() const { return Vertices; }
 	int32 Num() const { return Vertices.Num(); }
-	FSerializedIslandData GetSerializedData() const;
 
 	bool IsOperationAllowed(EGraphIslandOperations Op) const { return EnumHasAnyFlags(AllowedOperations, Op); }
 	void SetOperationAllowed(EGraphIslandOperations Op, bool bAllowed);
diff --git a/Engine/Plugins/Experimental/GameplayGraph/Source/GameplayGraph/Public/Graph/GraphSerialization.h b/Engine/Plugins/Experimental/GameplayGraph/Source/GameplayGraph/Public/Graph/GraphSerialization.h
new file mode 100644
--- /dev/null
+++ b/Engine/Plugins/Experimental/GameplayGraph/Source/GameplayGraph/Public/Graph/GraphSerialization.h
@@ -0,0 +1,52 @@
+// Copyright Epic Games, Inc. All Rights Reserved.
+#pragma once
+
+#include "Graph/GraphHandle.h"
+
+class UGraphVertex;
+class UGraphEdge;
+class UGraphIsland;
+
+struct FGraphProperties;
+struct FGraphVertexHandle;
+struct FGraphEdgeHandle;
+struct FGraphIslandHandle;
+
+class IGraphSerialization
+{
+public:
+	virtual ~IGraphSerialization() {}
+
+	virtual void Initialize(int32 NumVertices, int32 NumEdges, int32 NumIslands) = 0;
+	virtual void WriteGraphProperties(const FGraphProperties& Properties) = 0;
+	virtual void WriteGraphVertex(const FGraphVertexHandle& VertexHandle, const UGraphVertex* Vertex) = 0;
+	virtual void WriteGraphEdge(const FGraphEdgeHandle& EdgeHandle, const UGraphEdge* Edge) = 0;
+	virtual void WriteGraphIsland(const FGraphIslandHandle& IslandHandle, const UGraphIsland* Island) = 0;
+};
+
+class IGraphDeserialization
+{
+public:
+	virtual ~IGraphDeserialization() {}
+
+	virtual const FGraphProperties& GetProperties() const = 0;
+	virtual int32 NumVertices() const = 0;
+	virtual void ForEveryVertex(const TFunction<FGraphVertexHandle(const FGraphVertexHandle&)>& Lambda) const = 0;
+
+	struct FEdgeConstructionData
+	{
+		FGraphVertexHandle Vertex1;
+		FGraphVertexHandle Vertex2;
+	};
+
+	virtual int32 NumEdges() const = 0;
+	virtual void ForEveryEdge(const TFunction<FGraphEdgeHandle(const FGraphEdgeHandle&, const FEdgeConstructionData&)>& Lambda) const = 0;
+
+	struct FIslandConstructionData
+	{
+		TArray<FGraphVertexHandle> Vertices;
+	};
+
+	virtual int32 NumIslands() const = 0;
+	virtual void ForEveryIsland(const TFunction<FGraphIslandHandle(const FGraphIslandHandle&, const FIslandConstructionData&)>& Lambda) const = 0;
+};
diff --git a/Engine/Plugins/Experimental/GameplayGraph/Tests/Private/GraphEdgeTests.cpp b/Engine/Plugins/Experimental/GameplayGraph/Tests/Private/GraphEdgeTests.cpp
--- a/Engine/Plugins/Experimental/GameplayGraph/Tests/Private/GraphEdgeTests.cpp
+++ b/Engine/Plugins/Experimental/GameplayGraph/Tests/Private/GraphEdgeTests.cpp
@@ -31,10 +31,6 @@ TEST_CASE_METHOD(FTestGraphBuilder, "Graph::Edge::Create", "[graph][edge]")
 		CHECK(E->GetOtherNode(VertexHandles[0]) == VertexHandles[1]);
 		CHECK(E->GetOtherNode(VertexHandles[1]) == VertexHandles[0]);
 		CHECK(E->GetOtherNode(VertexHandles[2]) == FGraphVertexHandle{});
-
-		FSerializedEdgeData Serialized = E->GetSerializedData();
-		CHECK(Serialized.Node1 == VertexHandles[0]);
-		CHECK(Serialized.Node2 == VertexHandles[1]);
 	}
 
 	SECTION("Remove Edge")
diff --git a/Engine/Plugins/Experimental/GameplayGraph/Tests/Private/GraphIslandTests.cpp b/Engine/Plugins/Experimental/GameplayGraph/Tests/Private/GraphIslandTests.cpp
--- a/Engine/Plugins/Experimental/GameplayGraph/Tests/Private/GraphIslandTests.cpp
+++ b/Engine/Plugins/Experimental/GameplayGraph/Tests/Private/GraphIslandTests.cpp
@@ -717,26 +717,6 @@ TEST_CASE_METHOD(FTestGraphBuilder, "Graph::Island::Operation Allowed::Disable D
 	CHECK(Island->GetVertices().Contains(VertexHandles[3]) == true);
 }
 
-TEST_CASE_METHOD(FTestGraphBuilder, "Graph::Island::Get Serialized Data", "[graph][island]")
-{
-	PopulateVertices(4, true);
-	BuildLinearEdges(4);
-
-	REQUIRE(IslandHandles.Num() == 1);
-
-	UGraphIsland* Island = IslandHandles[0].GetIsland();
-	REQUIRE(Island != nullptr);
-
-	FSerializedIslandData Serialized = Island->GetSerializedData();
-	CHECK(Island->Num() == Serialized.Vertices.Num());
-	TSet<FGraphVertexHandle> SerializedVertexSet{ Serialized.Vertices };
-	CHECK(Island->GetVertices().Num() == SerializedVertexSet.Num());
-	for (const FGraphVertexHandle& Handle : SerializedVertexSet)
-	{
-		CHECK(Island->GetVertices().Contains(Handle));
-	}
-}
-
 TEST_CASE_METHOD(FTestGraphBuilder, "Graph::Island::Events::Vertex Added", "[graph][island]")
 {
 	PopulateVertices(2, true);
diff --git a/Engine/Plugins/Experimental/GameplayGraph/Tests/Private/GraphSerializationTests.cpp b/Engine/Plugins/Experimental/GameplayGraph/Tests/Private/GraphSerializationTests.cpp
--- a/Engine/Plugins/Experimental/GameplayGraph/Tests/Private/GraphSerializationTests.cpp
+++ b/Engine/Plugins/Experimental/GameplayGraph/Tests/Private/GraphSerializationTests.cpp
@@ -2,6 +2,8 @@
 #include "TestHarness.h"
 #include "TestGraphBuilder.h"
 
+#include "Graph/GraphDefaultSerialization.h"
+
 TEST_CASE_METHOD(FTestGraphBuilder, "Graph::Serialization::Write", "[graph]")
 {
 	PopulateVertices(6, true);
@@ -11,7 +13,9 @@ TEST_CASE_METHOD(FTestGraphBuilder, "Graph::Serialization::Write", "[graph]")
 	CHECK(Graph->NumEdges() == 4);
 	CHECK(Graph->NumIslands() == 2);
 
-	FSerializableGraph SerializedGraph = Graph->GetSerializableGraph();
+	FDefaultGraphSerialization Serializer;
+	Serializer << *Graph;
+	const FSerializableGraph& SerializedGraph = Serializer.GetData();
 	CHECK(SerializedGraph.Properties == Graph->GetProperties());
 	CHECK(SerializedGraph.Vertices.Num() == 6);
 	CHECK(SerializedGraph.Vertices.Contains(VertexHandles[0]) == true);
@@ -78,7 +82,9 @@ TEST_CASE_METHOD(FTestGraphBuilder, "Graph::Serialization::Read (normal)", "[gra
 
 	SerializedGraph.Islands.Add(FGraphIslandHandle { FGraphUniqueIndex::CreateUniqueIndex(), nullptr }, FSerializedIslandData { TArray{SerializedGraph.Vertices[0], SerializedGraph.Vertices[1], SerializedGraph.Vertices[2]} });
 	SerializedGraph.Islands.Add(FGraphIslandHandle { FGraphUniqueIndex::CreateUniqueIndex(), nullptr }, FSerializedIslandData { TArray{SerializedGraph.Vertices[3], SerializedGraph.Vertices[4], SerializedGraph.Vertices[5]} });
-	Graph->LoadFromSerializedGraph(SerializedGraph);
+
+	FDefaultGraphDeserialization Deserializer { SerializedGraph };
+	Deserializer >> *Graph;
 
 	CHECK(SerializedGraph.Properties == Graph->GetProperties());
 	CHECK(Graph->NumVertices() == 6);
@@ -135,7 +141,8 @@ TEST_CASE_METHOD(FTestGraphBuilder, "Graph::Serialization::Read::Errors::Vertex:
 	SerializedGraph.Properties.bGenerateIslands = true;
 	SerializedGraph.Vertices.Add(FGraphVertexHandle { FGraphUniqueIndex{ FGuid{0, 0, 0, 0} }, nullptr });
 
-	Graph->LoadFromSerializedGraph(SerializedGraph);
+	FDefaultGraphDeserialization Deserializer { SerializedGraph };
+	Deserializer >> *Graph;
 	CHECK(SerializedGraph.Properties == Graph->GetProperties());
 	CHECK(Graph->NumVertices() == 0);
 	CHECK(Graph->NumEdges() == 0);
@@ -150,7 +157,8 @@ TEST_CASE_METHOD(FTestGraphBuilder, "Graph::Serialization::Read::Errors::Edge::I
 	SerializedGraph.Vertices.Add(FGraphVertexHandle { FGraphUniqueIndex::CreateUniqueIndex(), nullptr });
 	SerializedGraph.Edges.Add(FGraphEdgeHandle { FGraphUniqueIndex{FGuid{0, 0, 0, 0}}, nullptr }, FSerializedEdgeData{ SerializedGraph.Vertices[0], SerializedGraph.Vertices[1] });
 
-	Graph->LoadFromSerializedGraph(SerializedGraph);
+	FDefaultGraphDeserialization Deserializer { SerializedGraph };
+	Deserializer >> *Graph;
 	CHECK(SerializedGraph.Properties == Graph->GetProperties());
 	CHECK(Graph->NumVertices() == 2);
 	CHECK(Graph->NumEdges() == 0);
@@ -165,7 +173,8 @@ TEST_CASE_METHOD(FTestGraphBuilder, "Graph::Serialization::Read::Errors::Edge::I
 	SerializedGraph.Vertices.Add(FGraphVertexHandle { FGraphUniqueIndex{FGuid{0, 0, 0, 2}}, nullptr });
 	SerializedGraph.Edges.Add(FGraphEdgeHandle { FGraphUniqueIndex::CreateUniqueIndex(), nullptr }, FSerializedEdgeData{ SerializedGraph.Vertices[0], FGraphVertexHandle { FGraphUniqueIndex{FGuid{0, 0, 0, 3}}, nullptr } });
 
-	Graph->LoadFromSerializedGraph(SerializedGraph);
+	FDefaultGraphDeserialization Deserializer { SerializedGraph };
+	Deserializer >> *Graph;
 	CHECK(SerializedGraph.Properties == Graph->GetProperties());
 	CHECK(Graph->NumVertices() == 2);
 	CHECK(Graph->NumEdges() == 0);
@@ -181,7 +190,8 @@ TEST_CASE_METHOD(FTestGraphBuilder, "Graph::Serialization::Read::Errors::Island:
 	SerializedGraph.Edges.Add(FGraphEdgeHandle { FGraphUniqueIndex::CreateUniqueIndex(), nullptr }, FSerializedEdgeData{ SerializedGraph.Vertices[0], SerializedGraph.Vertices[1] });
 	SerializedGraph.Islands.Add(FGraphIslandHandle { FGraphUniqueIndex{FGuid{0, 0, 0, 0}}, nullptr }, FSerializedIslandData { TArray{SerializedGraph.Vertices[0], SerializedGraph.Vertices[1] } });
 
-	Graph->LoadFromSerializedGraph(SerializedGraph);
+	FDefaultGraphDeserialization Deserializer { SerializedGraph };
+	Deserializer >> *Graph;
 	CHECK(SerializedGraph.Properties == Graph->GetProperties());
 	CHECK(Graph->NumVertices() == 2);
 	CHECK(Graph->NumEdges() == 1);
@@ -198,7 +208,8 @@ TEST_CASE_METHOD(FTestGraphBuilder, "Graph::Serialization::Read::Errors::Island:
 	SerializedGraph.Edges.Add(FGraphEdgeHandle { FGraphUniqueIndex::CreateUniqueIndex(), nullptr }, FSerializedEdgeData{ SerializedGraph.Vertices[0], SerializedGraph.Vertices[1] });
 	SerializedGraph.Islands.Add(FGraphIslandHandle { FGraphUniqueIndex{FGuid{1, 0, 0, 0}}, nullptr }, FSerializedIslandData { TArray{ SerializedGraph.Vertices[0], SerializedGraph.Vertices[1], FGraphVertexHandle { FGraphUniqueIndex{FGuid{0, 0, 0, 3}}, nullptr } } });
 
-	Graph->LoadFromSerializedGraph(SerializedGraph);
+	FDefaultGraphDeserialization Deserializer { SerializedGraph };
+	Deserializer >> *Graph;
 	CHECK(SerializedGraph.Properties == Graph->GetProperties());
 	CHECK(Graph->NumVertices() == 2);
 	CHECK(Graph->NumEdges() == 1);
@@ -226,7 +237,8 @@ TEST_CASE_METHOD(FTestGraphBuilder, "Graph::Serialization::Read::Errors::Island:
 	SerializedGraph.Islands.Add(FGraphIslandHandle { FGraphUniqueIndex{FGuid{1, 0, 0, 0}}, nullptr }, FSerializedIslandData { TArray { SerializedGraph.Vertices[0], SerializedGraph.Vertices[1] } });
 	SerializedGraph.Islands.Add(FGraphIslandHandle { FGraphUniqueIndex{FGuid{2, 0, 0, 0}}, nullptr }, FSerializedIslandData { TArray { SerializedGraph.Vertices[2] } });
 
-	Graph->LoadFromSerializedGraph(SerializedGraph);
+	FDefaultGraphDeserialization Deserializer { SerializedGraph };
+	Deserializer >> *Graph;
 	CHECK(SerializedGraph.Properties == Graph->GetProperties());
 	CHECK(Graph->NumVertices() == 3);
 	CHECK(Graph->NumEdges() == 2);
@@ -255,7 +267,8 @@ TEST_CASE_METHOD(FTestGraphBuilder, "Graph::Serialization::Read::Errors::Island:
 
 	SerializedGraph.Islands.Add(FGraphIslandHandle { FGraphUniqueIndex{FGuid{1, 0, 0, 0}}, nullptr }, FSerializedIslandData { TArray { SerializedGraph.Vertices[0], SerializedGraph.Vertices[1], SerializedGraph.Vertices[2] } });
 
-	Graph->LoadFromSerializedGraph(SerializedGraph);
+	FDefaultGraphDeserialization Deserializer { SerializedGraph };
+	Deserializer >> *Graph;
 	CHECK(SerializedGraph.Properties == Graph->GetProperties());
 	CHECK(Graph->NumVertices() == 3);
 	CHECK(Graph->NumEdges() == 1);
diff --git a/Engine/Plugins/Experimental/GameplayGraph/Tests/Private/GraphTests.cpp b/Engine/Plugins/Experimental/GameplayGraph/Tests/Private/GraphTests.cpp
--- a/Engine/Plugins/Experimental/GameplayGraph/Tests/Private/GraphTests.cpp
+++ b/Engine/Plugins/Experimental/GameplayGraph/Tests/Private/GraphTests.cpp
@@ -557,6 +557,28 @@ TEST_CASE_METHOD(FTestGraphBuilder, "Graph::Remove Vertex::Bulk::Linear", "[grap
 	CHECK(Graph->NumIslands() == 3);
 }
 
+TEST_CASE_METHOD(FTestGraphBuilder, "Graph::Remove Vertex::Incomplete Handle", "[graph]")
+{
+	PopulateVertices(5, true);
+	BuildLinearEdges(5);
+
+	FGraphVertexHandle RemoveHandle{ VertexHandles[0].GetUniqueIndex(), nullptr };
+	CHECK(Graph->NumVertices() == 5);
+	CHECK(Graph->NumEdges() == 4);
+	CHECK(Graph->NumIslands() == 1);
+
+	UGraphIsland* Island = IslandHandles[0].GetIsland();
+	REQUIRE(Island != nullptr);
+	CHECK(Island->Num() == 5);
+
+	Graph->RemoveVertex(RemoveHandle);
+
+	CHECK(Graph->NumVertices() == 4);
+	CHECK(Graph->NumEdges() == 3);
+	CHECK(Graph->NumIslands() == 1);
+	CHECK(Island->Num() == 4);
+}
+
 TEST_CASE_METHOD(FTestGraphBuilder, "Graph::Remove Island", "[graph]")
 {
 	PopulateVertices(10, true);
