remove concepts to appease non C++20 compilation

diff --git a/Engine/Plugins/Experimental/GameplayGraph/Source/GameplayGraph/Public/Graph/GraphDefaultSerialization.h b/Engine/Plugins/Experimental/GameplayGraph/Source/GameplayGraph/Public/Graph/GraphDefaultSerialization.h
--- a/Engine/Plugins/Experimental/GameplayGraph/Source/GameplayGraph/Public/Graph/GraphDefaultSerialization.h
+++ b/Engine/Plugins/Experimental/GameplayGraph/Source/GameplayGraph/Public/Graph/GraphDefaultSerialization.h
@@ -4,7 +4,6 @@
 #include "Graph/Graph.h"
 #include "Graph/GraphSerialization.h"
 
-#include <concepts>
 #include <type_traits>
 
 #include "GraphDefaultSerialization.generated.h"
@@ -63,10 +62,7 @@ struct FSerializableGraph
 	TMap<FGraphIslandHandle, FSerializedIslandData> Islands;
 };
 
-template<typename T>
-concept CDefaultGraphSerializable = std::is_base_of_v<FSerializableGraph, T>;
-
-template<CDefaultGraphSerializable TSerializableGraph>
+template<typename TSerializableGraph>
 class TDefaultGraphSerialization : public IGraphSerialization
 {
 public:
@@ -124,7 +120,7 @@ protected:
 	TSerializableGraph Data;
 };
 
-template<CDefaultGraphSerializable TSerializableGraph>
+template<typename TSerializableGraph>
 class TDefaultGraphDeserialization : public IGraphDeserialization
 {
 public:
