add additional checks and ensures into UGraphIsland::AddVertex and UGraphIsland::RemoveVertex to catch cases where a vertex could be in two islands at once

#tests manual

diff --git a/Engine/Plugins/Experimental/GameplayGraph/Source/GameplayGraph/Private/Graph/GraphIsland.cpp b/Engine/Plugins/Experimental/GameplayGraph/Source/GameplayGraph/Private/Graph/GraphIsland.cpp
--- a/Engine/Plugins/Experimental/GameplayGraph/Source/GameplayGraph/Private/Graph/GraphIsland.cpp
+++ b/Engine/Plugins/Experimental/GameplayGraph/Source/GameplayGraph/Private/Graph/GraphIsland.cpp
@@ -15,46 +15,45 @@ void UGraphIsland::Destroy()
 {
 	TRACE_CPUPROFILER_EVENT_SCOPE(UGraphIsland::Destroy);
 	TSet<FGraphVertexHandle> VertexCopy = Vertices;
-	Vertices.Empty();
 
+	Vertices.Empty();
 	// Need to work off a copy of the vertex handles since calling
 	// UGraph::RemoveVertex will attempt to remove the node from its
 	// parent island as well and that'll run into an error of modifying
 	// the Vertices TSet during iteration.
-	for (const FGraphVertexHandle& Node : VertexCopy)
+	if (UGraph* Graph = GetGraph())
 	{
-		if (TObjectPtr<UGraph> Graph = GetGraph())
-		{
-			Graph->RemoveVertex(Node);
-		}
+		Graph->RemoveBulkVertices(VertexCopy.Array());
 	}
 
-	Vertices.Empty();
 	HandleOnDestroyed();
 }
 
 void UGraphIsland::AddVertex(const FGraphVertexHandle& Node)
 {
 	TRACE_CPUPROFILER_EVENT_SCOPE(UGraphIsland::AddVertex);
-	if (!Node.IsValid())
+	if (!ensure(Node.IsValid()))
 	{
 		return;
 	}
 
-	TObjectPtr<UGraphVertex> NodePtr = Node.GetVertex();
-	if (!NodePtr)
+	UGraphVertex* NodePtr = Node.GetVertex();
+	if (!ensure(NodePtr))
 	{
 		return;
 	}
 
-	if (FGraphIslandHandle OldIslandHandle = NodePtr->GetParentIsland(); OldIslandHandle.IsComplete())
+	FGraphIslandHandle OldIslandHandle = NodePtr->GetParentIsland();
+	if (OldIslandHandle.IsValid())
 	{
-		if (TObjectPtr<UGraphIsland> OldIsland = OldIslandHandle.GetIsland())
+		UGraphIsland* OldIsland = OldIslandHandle.GetIsland();
+		if (ensure(OldIsland))
 		{
 			OldIsland->RemoveVertex(Node);
 		}
 	}
 
+	ensure(NodePtr->GetParentIsland() == FGraphIslandHandle{});
 	NodePtr->SetParentIsland(Handle());
 	Vertices.Add(Node);
 	HandleOnVertexAdded(Node);
@@ -63,17 +62,33 @@ void UGraphIsland::AddVertex(const FGraphVertexHandle& Node)
 void UGraphIsland::RemoveVertex(const FGraphVertexHandle& Node)
 {
 	TRACE_CPUPROFILER_EVENT_SCOPE(UGraphIsland::RemoveVertex);
-	if (!Node.IsValid())
+	if (!ensure(Node.IsValid()))
 	{
 		return;
 	}
 
-	if (UGraphVertex* NodePtr = Node.GetVertex())
+	const bool bIsInIslandSet = Vertices.Contains(Node);
+	bool bIsNodeParentSet = false;
+
+	UGraphVertex* NodePtr = Node.GetVertex();
+	if (ensure(NodePtr))
+	{
+		bIsNodeParentSet = NodePtr->GetParentIsland() == Handle();
+		if (bIsNodeParentSet)
+		{
+			NodePtr->SetParentIsland({});
+		}
+	}
+
+	if (bIsInIslandSet)
+	{
+		Vertices.Remove(Node);
+	}
+
+	if (bIsInIslandSet || bIsNodeParentSet)
 	{
-		NodePtr->SetParentIsland({});
+		HandleOnVertexRemoved(Node);
 	}
-	Vertices.Remove(Node);
-	HandleOnVertexRemoved(Node);
 }
 
 void UGraphIsland::HandleOnVertexAdded(const FGraphVertexHandle& Handle)
