gameplay graph should reserve its output space when serializing

- also added trace scopes in serialization and deserialization

diff --git a/Engine/Plugins/Experimental/GameplayGraph/Source/GameplayGraph/Private/Graph/Graph.cpp b/Engine/Plugins/Experimental/GameplayGraph/Source/GameplayGraph/Private/Graph/Graph.cpp
--- a/Engine/Plugins/Experimental/GameplayGraph/Source/GameplayGraph/Private/Graph/Graph.cpp
+++ b/Engine/Plugins/Experimental/GameplayGraph/Source/GameplayGraph/Private/Graph/Graph.cpp
@@ -786,108 +786,133 @@ void UGraph::ReserveIslands(int32 Delta)
 
 void operator<<(IGraphSerialization& Output, const UGraph& Graph)
 {
+	TRACE_CPUPROFILER_EVENT_SCOPE_STR("Gameplay Graph Serialization");
+	Output.Initialize(Graph.NumVertices(), Graph.NumEdges(), Graph.NumIslands());
 	Output.WriteGraphProperties(Graph.GetProperties());
-	for (const TPair<FGraphVertexHandle, TObjectPtr<UGraphVertex>>& Kvp : Graph.GetVertices())
+
 	{
-		if (!ensure(Kvp.Key.IsComplete() && Kvp.Value))
+		TRACE_CPUPROFILER_EVENT_SCOPE_STR("Serialize Vertices");
+		for (const TPair<FGraphVertexHandle, TObjectPtr<UGraphVertex>>& Kvp : Graph.GetVertices())
 		{
-			continue;
-		}
+			if (!ensure(Kvp.Key.IsComplete() && Kvp.Value))
+			{
+				continue;
+			}
 
-		Output.WriteGraphVertex(Kvp.Key, Kvp.Value);
+			Output.WriteGraphVertex(Kvp.Key, Kvp.Value);
+		}
 	}
 
-	for (const TPair<FGraphEdgeHandle, TObjectPtr<UGraphEdge>>& Kvp : Graph.GetEdges())
 	{
-		if (!ensure(Kvp.Key.IsComplete() && Kvp.Value))
+		TRACE_CPUPROFILER_EVENT_SCOPE_STR("Serialize Edges");
+		for (const TPair<FGraphEdgeHandle, TObjectPtr<UGraphEdge>>& Kvp : Graph.GetEdges())
 		{
-			continue;
-		}
+			if (!ensure(Kvp.Key.IsComplete() && Kvp.Value))
+			{
+				continue;
+			}
 
-		Output.WriteGraphEdge(Kvp.Key, Kvp.Value);
+			Output.WriteGraphEdge(Kvp.Key, Kvp.Value);
+		}
 	}
 
-	for (const TPair<FGraphIslandHandle, TObjectPtr<UGraphIsland>>& Kvp : Graph.GetIslands())
 	{
-		if (!ensure(Kvp.Key.IsComplete() && Kvp.Value))
+		TRACE_CPUPROFILER_EVENT_SCOPE_STR("Serialize Islands");
+		for (const TPair<FGraphIslandHandle, TObjectPtr<UGraphIsland>>& Kvp : Graph.GetIslands())
 		{
-			continue;
-		}
+			if (!ensure(Kvp.Key.IsComplete() && Kvp.Value))
+			{
+				continue;
+			}
 
-		Output.WriteGraphIsland(Kvp.Key, Kvp.Value);
+			Output.WriteGraphIsland(Kvp.Key, Kvp.Value);
+		}
 	}
 }
 
 void operator>>(const IGraphDeserialization& Input, UGraph& Graph)
 {
+	TRACE_CPUPROFILER_EVENT_SCOPE_STR("Gameplay Graph Deserialization");
 	Graph.InitializeFromProperties(Input.GetProperties());
 
-	Graph.ReserveVertices(Input.NumVertices());
-	Input.ForEveryVertex(
-		[&Graph](const FGraphVertexHandle& InHandle)
-		{
-			if (!ensure(InHandle.IsValid()))
+	{
+		TRACE_CPUPROFILER_EVENT_SCOPE_STR("Deserialize Vertices");
+		Graph.ReserveVertices(Input.NumVertices());
+		Input.ForEveryVertex(
+			[&Graph](const FGraphVertexHandle& InHandle)
 			{
-				return FGraphVertexHandle{};
-			}
+				if (!ensure(InHandle.IsValid()))
+				{
+					return FGraphVertexHandle{};
+				}
 
-			return Graph.CreateVertex(InHandle.GetUniqueIndex());
-		}
-	);
+				return Graph.CreateVertex(InHandle.GetUniqueIndex());
+			}
+		);
+	}
 
 	TArray<FGraphEdgeHandle> AllEdges;
 	AllEdges.Reserve(Input.NumEdges());
-	Graph.ReserveEdges(Input.NumEdges());
-	Input.ForEveryEdge(
-		[&Graph, &AllEdges](const FGraphEdgeHandle& InHandle, const IGraphDeserialization::FEdgeConstructionData& Data)
-		{
-			if (!ensure(InHandle.IsValid()))
+	{
+		TRACE_CPUPROFILER_EVENT_SCOPE_STR("Deserialize Edges");
+		Graph.ReserveEdges(Input.NumEdges());
+		Input.ForEveryEdge(
+			[&Graph, &AllEdges](const FGraphEdgeHandle& InHandle, const IGraphDeserialization::FEdgeConstructionData& Data)
 			{
-				return FGraphEdgeHandle{};
-			}
+				if (!ensure(InHandle.IsValid()))
+				{
+					return FGraphEdgeHandle{};
+				}
 
-			FGraphEdgeHandle Handle = Graph.CreateEdge(
-				Graph.GetCompleteNodeHandle(Data.Vertex1),
-				Graph.GetCompleteNodeHandle(Data.Vertex2),
-				InHandle.GetUniqueIndex(),
-				false
-			);
-			AllEdges.Add(Handle);
-			return Handle;
-		}
-	);
+				FGraphEdgeHandle Handle = Graph.CreateEdge(
+					Graph.GetCompleteNodeHandle(Data.Vertex1),
+					Graph.GetCompleteNodeHandle(Data.Vertex2),
+					InHandle.GetUniqueIndex(),
+					false
+				);
+				AllEdges.Add(Handle);
+				return Handle;
+			}
+		);
+	}
 
 	TArray<FGraphIslandHandle> AllIslands;
 	AllIslands.Reserve(Input.NumIslands());
-	Graph.ReserveIslands(Input.NumIslands());
-	Input.ForEveryIsland(
-		[&Graph, &AllIslands](const FGraphIslandHandle& InHandle, const IGraphDeserialization::FIslandConstructionData& Data)
-		{
-			if (!ensure(InHandle.IsValid()))
+
+
+	{
+		TRACE_CPUPROFILER_EVENT_SCOPE_STR("Deserialize Islands");
+		Graph.ReserveIslands(Input.NumIslands());
+		Input.ForEveryIsland(
+			[&Graph, &AllIslands](const FGraphIslandHandle& InHandle, const IGraphDeserialization::FIslandConstructionData& Data)
 			{
-				return FGraphIslandHandle{};
-			}
+				if (!ensure(InHandle.IsValid()))
+				{
+					return FGraphIslandHandle{};
+				}
 
-			TArray<FGraphVertexHandle> IslandVertices;
-			IslandVertices.Reserve(Data.Vertices.Num());
+				TArray<FGraphVertexHandle> IslandVertices;
+				IslandVertices.Reserve(Data.Vertices.Num());
 		
-			for (const FGraphVertexHandle& VertexHandle : Data.Vertices)
-			{
-				FGraphVertexHandle CompleteHandle = Graph.GetCompleteNodeHandle(VertexHandle);
-				if (CompleteHandle.IsComplete())
+				for (const FGraphVertexHandle& VertexHandle : Data.Vertices)
 				{
-					IslandVertices.Add(CompleteHandle);
+					FGraphVertexHandle CompleteHandle = Graph.GetCompleteNodeHandle(VertexHandle);
+					if (CompleteHandle.IsComplete())
+					{
+						IslandVertices.Add(CompleteHandle);
+					}
 				}
-			}
 
-			FGraphIslandHandle NewIslandHandle = Graph.CreateIsland(IslandVertices, InHandle.GetUniqueIndex());
-			AllIslands.Add(NewIslandHandle);
-			return NewIslandHandle;
-		}
-	);
+				FGraphIslandHandle NewIslandHandle = Graph.CreateIsland(IslandVertices, InHandle.GetUniqueIndex());
+				AllIslands.Add(NewIslandHandle);
+				return NewIslandHandle;
+			}
+		);
+	}
 
 	if (Graph::bFixupGraphOnLoad)
 	{
+		TRACE_CPUPROFILER_EVENT_SCOPE_STR("Post Deserialization Fixup");
 		Graph.MergeOrCreateIslands(AllEdges);
 
 		for (const FGraphIslandHandle& NewIslandHandle : AllIslands)
@@ -895,4 +920,4 @@ void operator>>(const IGraphDeserialization& Input, UGraph& Graph)
 			Graph.RemoveOrSplitIsland(NewIslandHandle.GetIsland());
 		}
 	}
-}
\ No newline at end of file
+}
