add a EGraphIslandConnectivityChange parameter to the FOnGraphIslandConnectedComponentsChanged event on graph islands to indicate what kind of change happened

diff --git a/Engine/Plugins/Experimental/GameplayGraph/Source/GameplayGraph/Private/Graph/Graph.cpp b/Engine/Plugins/Experimental/GameplayGraph/Source/GameplayGraph/Private/Graph/Graph.cpp
--- a/Engine/Plugins/Experimental/GameplayGraph/Source/GameplayGraph/Private/Graph/Graph.cpp
+++ b/Engine/Plugins/Experimental/GameplayGraph/Source/GameplayGraph/Private/Graph/Graph.cpp
@@ -600,9 +600,15 @@ void UGraph::RemoveOrSplitIsland(TObjectPtr<UGraphIsland> Island)
 		// This let's all the listeners of the islands make a decision on what to do with
 		// the island that had a destructive change to its connected components as well as
 		// the newly created islands.
-		for (TObjectPtr<UGraphIsland>& AffectedIsland : AffectedIslands)
+		for (UGraphIsland* AffectedIsland : AffectedIslands)
 		{
-			AffectedIsland->HandleOnConnectivityChanged();
+			if (!AffectedIsland)
+			{
+				continue;
+			}
+
+			const EGraphIslandConnectivityChange ChangeType = (AffectedIsland == Island) ? EGraphIslandConnectivityChange::SplitFrom : EGraphIslandConnectivityChange::SplitTo;
+			AffectedIsland->HandleOnConnectivityChanged(ChangeType);
 		}
 	}
 }
@@ -678,7 +684,7 @@ void UGraph::FinalizeVertex(const FGraphVertexHandle& InHandle)
 
 	if (TObjectPtr<UGraphIsland> Island = IslandHandle.GetIsland())
 	{
-		Island->HandleOnConnectivityChanged();
+		Island->HandleOnConnectivityChanged(EGraphIslandConnectivityChange::VertexAdd);
 	}
 }
 
diff --git a/Engine/Plugins/Experimental/GameplayGraph/Source/GameplayGraph/Private/Graph/GraphIsland.cpp b/Engine/Plugins/Experimental/GameplayGraph/Source/GameplayGraph/Private/Graph/GraphIsland.cpp
--- a/Engine/Plugins/Experimental/GameplayGraph/Source/GameplayGraph/Private/Graph/GraphIsland.cpp
+++ b/Engine/Plugins/Experimental/GameplayGraph/Source/GameplayGraph/Private/Graph/GraphIsland.cpp
@@ -106,9 +106,9 @@ void UGraphIsland::HandleOnDestroyed()
 	OnDestroyed.Broadcast(Handle());
 }
 
-void UGraphIsland::HandleOnConnectivityChanged()
+void UGraphIsland::HandleOnConnectivityChanged(EGraphIslandConnectivityChange Change)
 {
-	OnConnectivityChanged.Broadcast(Handle());
+	OnConnectivityChanged.Broadcast(Handle(), Change);
 }
 
 void UGraphIsland::SetOperationAllowed(EGraphIslandOperations Op, bool bAllowed)
diff --git a/Engine/Plugins/Experimental/GameplayGraph/Source/GameplayGraph/Public/Graph/GraphIsland.h b/Engine/Plugins/Experimental/GameplayGraph/Source/GameplayGraph/Public/Graph/GraphIsland.h
--- a/Engine/Plugins/Experimental/GameplayGraph/Source/GameplayGraph/Public/Graph/GraphIsland.h
+++ b/Engine/Plugins/Experimental/GameplayGraph/Source/GameplayGraph/Public/Graph/GraphIsland.h
@@ -15,19 +15,32 @@ enum class EGraphIslandOperations : int32
 {
 	None = 0,
 	Add = 1 << 0,
-Split = 1 << 1,
+	Split = 1 << 1,
 	Merge = 1 << 2,
 	Destroy = 1 << 3,
 	All = Add | Split | Merge | Destroy
 };
 ENUM_CLASS_FLAGS(EGraphIslandOperations);
 
+UENUM()
+enum class EGraphIslandConnectivityChange : int32
+{
+	// Vertex added into an island
+	VertexAdd,
+	// An island is split into 2 more islands
+	SplitFrom,
+	// An island was created by splitting an old island.
+	SplitTo,
+	// Some other undefined change - not used by the library but can be used by external users as a no-op of sorts.
+	Other
+};
+
 /** Delegate to track when some sort of batch change has occurred on this island that probably changes its connectivity.
  *  This is different from FOnGraphIslandNodeRemoved since FOnGraphIslandDestructiveChangeFinish will only be called
  *  once for the graph for a given operation while FOnGraphIslandNodeRemoved may be called multiple times if we're removing
  *  more than one node from the island at a given time. Note that this will only be called as a result of a destructive change.
  *  So repeatedly adding a node to an island won't call this event. */
-DECLARE_MULTICAST_DELEGATE_OneParam(FOnGraphIslandConnectedComponentsChanged, const FGraphIslandHandle&);
+DECLARE_MULTICAST_DELEGATE_TwoParams(FOnGraphIslandConnectedComponentsChanged, const FGraphIslandHandle&, EGraphIslandConnectivityChange);
 
 /** Delegate to track when this island should no longer exist. */
 DECLARE_MULTICAST_DELEGATE_OneParam(FOnGraphIslandDestroyed, const FGraphIslandHandle&);
@@ -83,7 +96,7 @@ protected:
 	virtual void HandleOnVertexAdded(const FGraphVertexHandle& Handle);
 	virtual void HandleOnVertexRemoved(const FGraphVertexHandle& Handle);
 	virtual void HandleOnDestroyed();
-	virtual void HandleOnConnectivityChanged();
+	virtual void HandleOnConnectivityChanged(EGraphIslandConnectivityChange Change);
 
 	/** Called when removing the island from the graph. */
 	void Destroy();
diff --git a/Engine/Plugins/Experimental/GameplayGraph/Tests/Private/GraphIslandTests.cpp b/Engine/Plugins/Experimental/GameplayGraph/Tests/Private/GraphIslandTests.cpp
--- a/Engine/Plugins/Experimental/GameplayGraph/Tests/Private/GraphIslandTests.cpp
+++ b/Engine/Plugins/Experimental/GameplayGraph/Tests/Private/GraphIslandTests.cpp
@@ -662,7 +662,7 @@ TEST_CASE_METHOD(FTestGraphBuilder, "Graph::Island::Events::On Connectivity Chan
 				REQUIRE(Island != nullptr);
 
 				Island->OnConnectivityChanged.AddLambda(
-					[&CallbackIslandHandle](const FGraphIslandHandle& InCallbackIslandHandle)
+					[&CallbackIslandHandle](const FGraphIslandHandle& InCallbackIslandHandle, EGraphIslandConnectivityChange Change)
 					{
 						CallbackIslandHandle = InCallbackIslandHandle;
 					}
@@ -689,7 +689,7 @@ TEST_CASE_METHOD(FTestGraphBuilder, "Graph::Island::Events::On Connectivity Chan
 
 		FGraphIslandHandle CallbackIslandHandle;
 		Island->OnConnectivityChanged.AddLambda(
-			[&CallbackIslandHandle](const FGraphIslandHandle& InCallbackIslandHandle)
+			[&CallbackIslandHandle](const FGraphIslandHandle& InCallbackIslandHandle, EGraphIslandConnectivityChange Change)
 			{
 				CallbackIslandHandle = InCallbackIslandHandle;
 			}
