gameplay graph incremental serialization

diff --git a/Engine/Plugins/Experimental/GameplayGraph/Source/GameplayGraph/Private/Graph/Graph.cpp b/Engine/Plugins/Experimental/GameplayGraph/Source/GameplayGraph/Private/Graph/Graph.cpp
--- a/Engine/Plugins/Experimental/GameplayGraph/Source/GameplayGraph/Private/Graph/Graph.cpp
+++ b/Engine/Plugins/Experimental/GameplayGraph/Source/GameplayGraph/Private/Graph/Graph.cpp
@@ -133,6 +133,7 @@ bool UGraph::CreateEdge(FGraphVertexHandle Node1, FGraphVertexHandle Node2, bool
 	// Also need to make sure the nodes are aware of the new edge.
 	Node1Ptr->AddEdgeTo(Node2);
 	Node2Ptr->AddEdgeTo(Node1);
+	OnEdgeCreated.Broadcast(FEdgeSpecifier{Node1, Node2});
 
 	// If we want to keep track of islands, this is where we need to create/merge islands.
 	if (Properties.bGenerateIslands)
@@ -536,6 +537,8 @@ void UGraph::RemoveEdgeInternal(const FGraphVertexHandle& VertexHandleA, const F
 			IslandHandle = NodeB->GetParentIsland();
 		}
 	}
+
+	OnEdgeRemoved.Broadcast(FEdgeSpecifier{VertexHandleA, VertexHandleB});
 	
 	// Removing an edge should cause the island to check if it needs to split.
 	if (Properties.bGenerateIslands && bHandleIslands)
diff --git a/Engine/Plugins/Experimental/GameplayGraph/Source/GameplayGraph/Private/Graph/GraphDefaultSerialization.cpp b/Engine/Plugins/Experimental/GameplayGraph/Source/GameplayGraph/Private/Graph/GraphDefaultSerialization.cpp
--- a/Engine/Plugins/Experimental/GameplayGraph/Source/GameplayGraph/Private/Graph/GraphDefaultSerialization.cpp
+++ b/Engine/Plugins/Experimental/GameplayGraph/Source/GameplayGraph/Private/Graph/GraphDefaultSerialization.cpp
@@ -2,4 +2,124 @@
 
 #include "Graph/GraphDefaultSerialization.h"
 
-#include UE_INLINE_GENERATED_CPP_BY_NAME(GraphDefaultSerialization)
\ No newline at end of file
+#include UE_INLINE_GENERATED_CPP_BY_NAME(GraphDefaultSerialization)
+
+bool operator==(const FSerializedEdgeData& Lhs, const FSerializedEdgeData& Rhs)
+{
+	// We can't really trust that the serialized edge data is in some canonical form so
+	// given that we're dealing with undirected graphs, I don't trust that the serialized data
+	// will always be in canonical form (where Node1 < Node2) so the equality operation has to
+	// take that into account.
+	if (Lhs.Node1 == Rhs.Node1)
+	{
+		return Lhs.Node2 == Rhs.Node2;
+	}
+	else if (Lhs.Node1 == Rhs.Node2)
+	{
+		return Lhs.Node2 == Rhs.Node1;
+	}
+
+	return false;
+}
+
+bool operator!=(const FSerializedEdgeData& Lhs, const FSerializedEdgeData& Rhs)
+{
+	// Is the default operator!= defined via operator==? In that case this might be redundant.
+	return !(Lhs == Rhs);
+}
+
+bool operator==(const FSerializedIslandData& Lhs, const FSerializedIslandData& Rhs)
+{
+	if (Lhs.Vertices.Num() != Rhs.Vertices.Num())
+	{
+		return false;
+	}
+
+	// This should be faster than going through Lhs.Vertices and checking if Rhs.Vertices contains the vertex.
+	// In that situation you'd be doing a O(N^2) algorithm since the common case would be the Lhs == Rhs and you'd
+	// have to do a linear sweep to find the vertex.
+	//
+	// In this case, building the set is O(N log N). Then doing N removals from the set is another O(N log N) operation.
+	TSet<FGraphVertexHandle> EqualitySet{ Lhs.Vertices };
+	ensure(EqualitySet.Num() == Lhs.Vertices.Num());
+
+	for (const FGraphVertexHandle& VertexHandle : Rhs.Vertices)
+	{
+		if (!EqualitySet.Contains(VertexHandle))
+		{
+			return false;
+		}
+	}
+
+	return true;
+}
+
+bool operator!=(const FSerializedIslandData& Lhs, const FSerializedIslandData& Rhs)
+{
+	return !(Lhs == Rhs);
+}
+
+bool operator==(const FSerializableGraph& Lhs, const FSerializableGraph& Rhs)
+{
+	if (Lhs.Properties != Rhs.Properties)
+	{
+		return false;
+	}
+
+	if (Lhs.Vertices.Num() != Rhs.Vertices.Num())
+	{
+		return false;
+	}
+
+	{
+		TSet<FGraphVertexHandle> EqualitySet { Lhs.Vertices };
+		ensure(EqualitySet.Num() == Lhs.Vertices.Num());
+		// Assumptions here: there are no duplicate vertices in Lhs hence why we don't check in the other direction.
+		// So given that we've already check the size of the array, this should guarantee equality (barring the duplicate vertex condition).
+		for (const FGraphVertexHandle& VertexHandle : Rhs.Vertices)
+		{
+			if (!EqualitySet.Contains(VertexHandle))
+			{
+				return false;
+			}
+		}
+	}
+
+	if (Lhs.Edges.Num() != Rhs.Edges.Num())
+	{
+		return false;
+	}
+
+	{
+		TSet<FSerializedEdgeData> EqualitySet { Lhs.Edges };
+		ensure(EqualitySet.Num() == Lhs.Edges.Num());
+		for (const FSerializedEdgeData& EdgeData : Rhs.Edges)
+		{
+			if (!EqualitySet.Contains(EdgeData))
+			{
+				return false;
+			}
+		}
+	}
+
+	if (Lhs.Islands.Num() != Rhs.Islands.Num())
+	{
+		return false;
+	}
+
+	for (const TPair<FGraphIslandHandle, FSerializedIslandData>& Kvp : Lhs.Islands)
+	{
+		const FSerializedIslandData* RhsData = Rhs.Islands.Find(Kvp.Key);
+		if (!RhsData)
+		{
+			return false;
+		}
+
+		if (Kvp.Value != *RhsData)
+		{
+			return false;
+		}
+	}
+
+	return true;
+}
\ No newline at end of file
diff --git a/Engine/Plugins/Experimental/GameplayGraph/Source/GameplayGraph/Private/Graph/GraphVertex.cpp b/Engine/Plugins/Experimental/GameplayGraph/Source/GameplayGraph/Private/Graph/GraphVertex.cpp
--- a/Engine/Plugins/Experimental/GameplayGraph/Source/GameplayGraph/Private/Graph/GraphVertex.cpp
+++ b/Engine/Plugins/Experimental/GameplayGraph/Source/GameplayGraph/Private/Graph/GraphVertex.cpp
@@ -32,11 +32,11 @@ void UGraphVertex::RemoveEdge(const FGraphVertexHandle& AdjacentVertexHandle)
 
 void UGraphVertex::HandleOnVertexRemoved()
 {
-	OnVertexRemoved.Broadcast();
+	OnVertexRemoved.Broadcast(Handle());
 }
 
 void UGraphVertex::SetParentIsland(const FGraphIslandHandle& Island)
 {
 	ParentIsland = Island;
-	OnParentIslandSet.Broadcast(Island);
+	OnParentIslandSet.Broadcast(Handle(), Island);
 }
\ No newline at end of file
diff --git a/Engine/Plugins/Experimental/GameplayGraph/Source/GameplayGraph/Public/Graph/Graph.h b/Engine/Plugins/Experimental/GameplayGraph/Source/GameplayGraph/Public/Graph/Graph.h
--- a/Engine/Plugins/Experimental/GameplayGraph/Source/GameplayGraph/Public/Graph/Graph.h
+++ b/Engine/Plugins/Experimental/GameplayGraph/Source/GameplayGraph/Public/Graph/Graph.h
@@ -53,6 +53,9 @@ private:
 	FGraphVertexHandle VertexHandle2;
 };
 
+DECLARE_MULTICAST_DELEGATE_OneParam(FOnGraphEdgeCreated, const FEdgeSpecifier&);
+DECLARE_MULTICAST_DELEGATE_OneParam(FOnGraphEdgeRemoved, const FEdgeSpecifier&);
+
 /**
  * A UGraph is a collection of nodes and edges. This graph representation
  * is meant to be easily integrable into gameplay systems in the Unreal Engine.
@@ -122,6 +125,8 @@ public:
 
 	FOnGraphVertexCreated OnVertexCreated;
 	FOnGraphIslandCreated OnIslandCreated;
+	FOnGraphEdgeCreated OnEdgeCreated;
+	FOnGraphEdgeRemoved OnEdgeRemoved;
 
 	const TMap<FGraphVertexHandle, TObjectPtr<UGraphVertex>>& GetVertices() const { return Vertices; }
 	const TMap<FGraphIslandHandle, TObjectPtr<UGraphIsland>>& GetIslands() const { return Islands; }
diff --git a/Engine/Plugins/Experimental/GameplayGraph/Source/GameplayGraph/Public/Graph/GraphDefaultSerialization.h b/Engine/Plugins/Experimental/GameplayGraph/Source/GameplayGraph/Public/Graph/GraphDefaultSerialization.h
--- a/Engine/Plugins/Experimental/GameplayGraph/Source/GameplayGraph/Public/Graph/GraphDefaultSerialization.h
+++ b/Engine/Plugins/Experimental/GameplayGraph/Source/GameplayGraph/Public/Graph/GraphDefaultSerialization.h
@@ -2,6 +2,7 @@
 #pragma once
 
 #include "Graph/Graph.h"
+#include "Graph/GraphIncrementalSerialization.h"
 #include "Graph/GraphSerialization.h"
 
 #include <type_traits>
@@ -20,8 +21,13 @@ struct FSerializedEdgeData
 	FGraphVertexHandle Node2;
 
 	// Comparison operators
-	friend bool operator==(const FSerializedEdgeData& Lhs, const FSerializedEdgeData& Rhs) = default;
-	friend bool operator!=(const FSerializedEdgeData& Lhs, const FSerializedEdgeData& Rhs) = default;
+	friend GAMEPLAYGRAPH_API bool operator==(const FSerializedEdgeData& Lhs, const FSerializedEdgeData& Rhs);
+	friend GAMEPLAYGRAPH_API bool operator!=(const FSerializedEdgeData& Lhs, const FSerializedEdgeData& Rhs);
+
+	friend uint32 GAMEPLAYGRAPH_API GetTypeHash(const FSerializedEdgeData& EdgeData)
+	{
+		return HashCombine(GetTypeHash(EdgeData.Node1), GetTypeHash(EdgeData.Node2));
+	}
 };
 
 USTRUCT()
@@ -33,8 +39,8 @@ struct FSerializedIslandData
 	TArray<FGraphVertexHandle> Vertices;
 
 	// Comparison operators
-	friend bool operator==(const FSerializedIslandData& Lhs, const FSerializedIslandData& Rhs) = default;
-	friend bool operator!=(const FSerializedIslandData& Lhs, const FSerializedIslandData& Rhs) = default;
+	friend GAMEPLAYGRAPH_API bool operator==(const FSerializedIslandData& Lhs, const FSerializedIslandData& Rhs);
+	friend GAMEPLAYGRAPH_API bool operator!=(const FSerializedIslandData& Lhs, const FSerializedIslandData& Rhs);
 };
 
 /**
@@ -60,11 +66,14 @@ struct FSerializableGraph
 
 	UPROPERTY(SaveGame)
 	TMap<FGraphIslandHandle, FSerializedIslandData> Islands;
+
+	friend GAMEPLAYGRAPH_API bool operator==(const FSerializableGraph& Lhs, const FSerializableGraph& Rhs);
 };
 
 template<typename TSerializableGraph>
 class TDefaultGraphSerialization : public IGraphSerialization
 {
+	static_assert(std::is_base_of_v<FSerializableGraph, TSerializableGraph>, "TDefaultGraphSerialization must serialize a type of FSerializableGraph");
 public:
 	virtual ~TDefaultGraphSerialization() = default;
 
@@ -111,6 +120,13 @@ public:
 		Data.Islands.Emplace(IslandHandle, MoveTemp(Serialized));
 	}
 
+	virtual void Reset() override
+	{
+		Data.Vertices.Reset();
+		Data.Edges.Reset();
+		Data.Islands.Reset();
+	}
+
 protected:
 	TSerializableGraph Data;
 };
@@ -118,6 +134,7 @@ protected:
 template<typename TSerializableGraph>
 class TDefaultGraphDeserialization : public IGraphDeserialization
 {
+	static_assert(std::is_base_of_v<FSerializableGraph, TSerializableGraph>, "TDefaultGraphDeserialization must deserialize a type of FSerializableGraph");
 public:
 	explicit TDefaultGraphDeserialization(const TSerializableGraph& InData)
 	: Data(InData)
@@ -198,4 +215,300 @@ protected:
 };
 
 using FDefaultGraphSerialization = TDefaultGraphSerialization<FSerializableGraph>;
-using FDefaultGraphDeserialization = TDefaultGraphDeserialization<FSerializableGraph>;
\ No newline at end of file
+using FDefaultGraphDeserialization = TDefaultGraphDeserialization<FSerializableGraph>;
+
+//
+// DELTA ACTIONS
+//
+enum class EDefaultDeltaActionType
+{
+	Add,
+	Remove
+};
+
+template<EDefaultDeltaActionType Type>
+struct TDefaultGraphVertexDeltaAction : public TDerivedGraphDeltaAction<TDefaultGraphVertexDeltaAction<Type>>
+{
+public:
+	explicit TDefaultGraphVertexDeltaAction(const FGraphVertexHandle& InHandle)
+	: VertexHandle(InHandle)
+	{}
+
+	FGraphVertexHandle VertexHandle;
+};
+
+template<EDefaultDeltaActionType Type>
+struct TDefaultGraphEdgeDeltaAction : public TDerivedGraphDeltaAction<TDefaultGraphEdgeDeltaAction<Type>>
+{
+public:
+	explicit TDefaultGraphEdgeDeltaAction(const FEdgeSpecifier& InEdge)
+	: Edge(InEdge)
+	{}
+
+	FEdgeSpecifier Edge;
+};
+
+template<EDefaultDeltaActionType Type>
+struct TDefaultGraphIslandDeltaAction : public TDerivedGraphDeltaAction<TDefaultGraphIslandDeltaAction<Type>>
+{
+public:
+	explicit TDefaultGraphIslandDeltaAction(const FGraphIslandHandle& InHandle)
+	: IslandHandle(InHandle)
+	{}
+
+	FGraphIslandHandle IslandHandle;
+};
+
+
+template<EDefaultDeltaActionType Type>
+struct TDefaultGraphIslandVertexDeltaAction : public TDerivedGraphDeltaAction<TDefaultGraphIslandVertexDeltaAction<Type>>
+{
+public:
+	explicit TDefaultGraphIslandVertexDeltaAction(const FGraphIslandHandle& InIslandHandle, const FGraphVertexHandle& InVertexHandle)
+	: IslandHandle(InIslandHandle)
+	, VertexHandle(InVertexHandle)
+	{}
+
+	FGraphIslandHandle IslandHandle;
+	FGraphVertexHandle VertexHandle;
+};
+
+//
+// DELTA HANDLER
+//
+template<typename TInSerializableGraph>
+class TDefaultGraphDeltaActionHandler : public TGraphDeltaActionHandler<TInSerializableGraph>
+{
+	using TThisClass = TDefaultGraphDeltaActionHandler<TInSerializableGraph>;
+	static_assert(std::is_base_of_v<FSerializableGraph, TInSerializableGraph>, "TDefaultGraphDeltaActionHandler's TSerializableGraph template parameter must derive from FSerializableGraph.");
+public:
+	using FCreateVertexDeltaAction = TDefaultGraphVertexDeltaAction<EDefaultDeltaActionType::Add>;
+	using FRemoveVertexDeltaAction = TDefaultGraphVertexDeltaAction<EDefaultDeltaActionType::Remove>;
+
+	using FCreateEdgeDeltaAction = TDefaultGraphEdgeDeltaAction<EDefaultDeltaActionType::Add>;
+	using FRemoveEdgeDeltaAction = TDefaultGraphEdgeDeltaAction<EDefaultDeltaActionType::Remove>;
+
+	using FCreateIslandDeltaAction = TDefaultGraphIslandDeltaAction<EDefaultDeltaActionType::Add>;
+	using FRemoveIslandDeltaAction = TDefaultGraphIslandDeltaAction<EDefaultDeltaActionType::Remove>;
+
+	using FAddIslandVertexDeltaAction = TDefaultGraphIslandVertexDeltaAction<EDefaultDeltaActionType::Add>;
+	using FRemoveIslandVertexDeltaAction = TDefaultGraphIslandVertexDeltaAction<EDefaultDeltaActionType::Remove>;
+
+	virtual void InitializeFromGraph(const TInSerializableGraph& InGraph) override
+	{
+		IncrementalVertices.Append(InGraph.Vertices);
+		IncrementalEdges.Append(InGraph.Edges);
+	}
+
+	virtual void Flush(TInSerializableGraph& OutGraph) override
+	{
+		OutGraph.Vertices = IncrementalVertices.Array();
+		OutGraph.Edges = IncrementalEdges.Array();
+	}
+
+	void Visit(const FCreateVertexDeltaAction& Action, TInSerializableGraph& OutGraph)
+	{
+		TRACE_CPUPROFILER_EVENT_SCOPE(TDefaultGraphDeltaActionHandler::VisitCreateVertex);
+		IncrementalVertices.Add(Action.VertexHandle);
+	}
+
+	void Visit(const FRemoveVertexDeltaAction& Action, TInSerializableGraph& OutGraph)
+	{
+		TRACE_CPUPROFILER_EVENT_SCOPE(TDefaultGraphDeltaActionHandler::VisitRemoveVertex);
+		IncrementalVertices.Remove(Action.VertexHandle);
+	}
+
+	void Visit(const FCreateEdgeDeltaAction& Action, TInSerializableGraph& OutGraph)
+	{
+		TRACE_CPUPROFILER_EVENT_SCOPE(TDefaultGraphDeltaActionHandler::VisitCreateEdge);
+		FSerializedEdgeData Data;
+		Data.Node1 = Action.Edge.GetVertexHandle1();
+		Data.Node2 = Action.Edge.GetVertexHandle2();
+		IncrementalEdges.Emplace(MoveTemp(Data));
+	}
+
+	void Visit(const FRemoveEdgeDeltaAction& Action, TInSerializableGraph& OutGraph)
+	{
+		TRACE_CPUPROFILER_EVENT_SCOPE(TDefaultGraphDeltaActionHandler::VisitRemoveEdge);
+		FSerializedEdgeData Data;
+		Data.Node1 = Action.Edge.GetVertexHandle1();
+		Data.Node2 = Action.Edge.GetVertexHandle2();
+		IncrementalEdges.Remove(Data);
+	}
+
+	void Visit(const FCreateIslandDeltaAction& Action, TInSerializableGraph& OutGraph)
+	{
+		TRACE_CPUPROFILER_EVENT_SCOPE(TDefaultGraphDeltaActionHandler::VisitCreateIsland);
+		// We're guaranteed that the island created event is fired before vertices are added into it.
+		// Hence we're safe to just create an empty FSerializedIslandData here.
+		OutGraph.Islands.Emplace(Action.IslandHandle, FSerializedIslandData{});
+	}
+
+	void Visit(const FRemoveIslandDeltaAction& Action, TInSerializableGraph& OutGraph)
+	{
+		TRACE_CPUPROFILER_EVENT_SCOPE(TDefaultGraphDeltaActionHandler::VisitRemoveIsland);
+		OutGraph.Islands.Remove(Action.IslandHandle);
+	}
+
+	void Visit(const FAddIslandVertexDeltaAction& Action, TInSerializableGraph& OutGraph)
+	{
+		TRACE_CPUPROFILER_EVENT_SCOPE(TDefaultGraphDeltaActionHandler::VisitAddIslandVertex);
+		if (FSerializedIslandData* Data = OutGraph.Islands.Find(Action.IslandHandle))
+		{
+			Data->Vertices.AddUnique(Action.VertexHandle);
+		}
+	}
+
+	void Visit(const FRemoveIslandVertexDeltaAction& Action, TInSerializableGraph& OutGraph)
+	{
+		TRACE_CPUPROFILER_EVENT_SCOPE(TDefaultGraphDeltaActionHandler::VisitRemoveIslandVertex);
+		if (FSerializedIslandData* Data = OutGraph.Islands.Find(Action.IslandHandle))
+		{
+			Data->Vertices.RemoveSwap(Action.VertexHandle, EAllowShrinking::No);
+		}
+	}
+
+private:
+	// Doing removes from the vertex and edges arrays is too slow so we make sets that we apply operations to that we copy over to the array in Flush.
+	TSet<FGraphVertexHandle> IncrementalVertices;
+	TSet<FSerializedEdgeData> IncrementalEdges;
+};
+
+using FDefaultGraphDeltaActionHandler = TDefaultGraphDeltaActionHandler<FSerializableGraph>;
+
+//
+// INCREMENTAL SERIALIZER
+//
+template<typename TSerializer, typename TDeltaActionHandler>
+class TDefaultGraphIncrementalSerialization : public TGraphIncrementalSerialization<TSerializer, TDeltaActionHandler>
+{
+	static_assert(std::is_base_of_v<TDefaultGraphSerialization<typename TDeltaActionHandler::TSerializableGraph>, TSerializer>, "TDefaultGraphIncrementalSerialization's TSerializer template parameter must derived from TDefaultGraphSerialization.");
+	static_assert(std::is_base_of_v<TDefaultGraphDeltaActionHandler<typename TDeltaActionHandler::TSerializableGraph>, TDeltaActionHandler>, "TDefaultGraphIncrementalSerialization's TDeltaActionHandler template parameter should inherit from TDefaultGraphDeltaActionHandler.");
+public:
+
+	explicit TDefaultGraphIncrementalSerialization(UGraph* InGraph)
+	: TGraphIncrementalSerialization<TSerializer, TDeltaActionHandler>(InGraph)
+	{
+		if (InGraph)
+		{
+			InGraph->OnVertexCreated.AddRaw(this, &TDefaultGraphIncrementalSerialization::OnGraphVertexCreated);
+			InGraph->OnEdgeCreated.AddRaw(this, &TDefaultGraphIncrementalSerialization::OnGraphEdgeCreated);
+			InGraph->OnEdgeRemoved.AddRaw(this, &TDefaultGraphIncrementalSerialization::OnGraphEdgeRemoved);
+			InGraph->OnIslandCreated.AddRaw(this, &TDefaultGraphIncrementalSerialization::OnGraphIslandCreated);
+
+			for (const TPair<FGraphVertexHandle, TObjectPtr<UGraphVertex>>& Kvp : InGraph->GetVertices())
+			{
+				StartListenToVertexChanges(Kvp.Key);
+			}
+
+			for (const TPair<FGraphIslandHandle, TObjectPtr<UGraphIsland>>& Kvp : InGraph->GetIslands())
+			{
+				StartListenToIslandChanges(Kvp.Key);
+			}
+		}
+	}
+
+	~TDefaultGraphIncrementalSerialization()
+	{
+		if (UGraph* Graph = this->GetGraph())
+		{
+			Graph->OnVertexCreated.RemoveAll(this);
+			Graph->OnEdgeCreated.RemoveAll(this);
+			Graph->OnEdgeRemoved.RemoveAll(this);
+			Graph->OnIslandCreated.RemoveAll(this);
+
+			for (const TPair<FGraphVertexHandle, TObjectPtr<UGraphVertex>>& Kvp : Graph->GetVertices())
+			{
+				StopListenToVertexChanges(Kvp.Key);
+			}
+
+			for (const TPair<FGraphIslandHandle, TObjectPtr<UGraphIsland>>& Kvp : Graph->GetIslands())
+			{
+				StopListenToIslandChanges(Kvp.Key);
+			}
+		}
+	}
+
+protected:
+	virtual void OnGraphVertexCreated(const FGraphVertexHandle& VertexHandle)
+	{
+		this->template AddDeltaAction<typename TDeltaActionHandler::FCreateVertexDeltaAction>(VertexHandle);
+		StartListenToVertexChanges(VertexHandle);
+	}
+
+	virtual void OnGraphVertexRemoved(const FGraphVertexHandle& VertexHandle)
+	{
+		this->template AddDeltaAction<typename TDeltaActionHandler::FRemoveVertexDeltaAction>(VertexHandle);
+		StopListenToVertexChanges(VertexHandle);
+	}
+
+	virtual void OnGraphEdgeCreated(const FEdgeSpecifier& Edge)
+	{
+		this->template AddDeltaAction<typename TDeltaActionHandler::FCreateEdgeDeltaAction>(Edge);
+	}
+
+	virtual void OnGraphEdgeRemoved(const FEdgeSpecifier& Edge)
+	{
+		this->template AddDeltaAction<typename TDeltaActionHandler::FRemoveEdgeDeltaAction>(Edge);
+	}
+
+	virtual void OnGraphIslandCreated(const FGraphIslandHandle& IslandHandle)
+	{
+		this->template AddDeltaAction<typename TDeltaActionHandler::FCreateIslandDeltaAction>(IslandHandle);
+		StartListenToIslandChanges(IslandHandle);
+	}
+
+	virtual void OnGraphIslandRemoved(const FGraphIslandHandle& IslandHandle)
+	{
+		this->template AddDeltaAction<typename TDeltaActionHandler::FRemoveIslandDeltaAction>(IslandHandle);
+		StopListenToIslandChanges(IslandHandle);
+	}
+
+	virtual void OnGraphIslandVertexAdded(const FGraphIslandHandle& IslandHandle, const FGraphVertexHandle& VertexHandle)
+	{
+		this->template AddDeltaAction<typename TDeltaActionHandler::FAddIslandVertexDeltaAction>(IslandHandle, VertexHandle);
+	}
+
+	virtual void OnGraphIslandVertexRemoved(const FGraphIslandHandle& IslandHandle, const FGraphVertexHandle& VertexHandle)
+	{
+		this->template AddDeltaAction<typename TDeltaActionHandler::FRemoveIslandVertexDeltaAction>(IslandHandle, VertexHandle);
+	}
+
+	virtual void StartListenToVertexChanges(const FGraphVertexHandle& VertexHandle)
+	{
+		if (UGraphVertex* Vertex = VertexHandle.GetVertex())
+		{
+			Vertex->OnVertexRemoved.AddRaw(this, &TDefaultGraphIncrementalSerialization::OnGraphVertexRemoved);
+		}
+	}
+
+	virtual void StopListenToVertexChanges(const FGraphVertexHandle& VertexHandle)
+	{
+		if (UGraphVertex* Vertex = VertexHandle.GetVertex())
+		{
+			Vertex->OnVertexRemoved.RemoveAll(this);
+		}
+	}
+
+	virtual void StartListenToIslandChanges(const FGraphIslandHandle& IslandHandle)
+	{
+		if (UGraphIsland* Island = IslandHandle.GetIsland())
+		{
+			Island->OnDestroyed.AddRaw(this, &TDefaultGraphIncrementalSerialization::OnGraphIslandRemoved);
+			Island->OnVertexAdded.AddRaw(this, &TDefaultGraphIncrementalSerialization::OnGraphIslandVertexAdded);
+			Island->OnVertexRemoved.AddRaw(this, &TDefaultGraphIncrementalSerialization::OnGraphIslandVertexRemoved);
+		}
+	}
+
+	virtual void StopListenToIslandChanges(const FGraphIslandHandle& IslandHandle)
+	{
+		if (UGraphIsland* Island = IslandHandle.GetIsland())
+		{
+			Island->OnDestroyed.RemoveAll(this);
+			Island->OnVertexAdded.RemoveAll(this);
+			Island->OnVertexRemoved.RemoveAll(this);
+		}
+	}
+};
+
+using FDefaultGraphIncrementalSerialization = TDefaultGraphIncrementalSerialization<FDefaultGraphSerialization, FDefaultGraphDeltaActionHandler>;
\ No newline at end of file
diff --git a/Engine/Plugins/Experimental/GameplayGraph/Source/GameplayGraph/Public/Graph/GraphIncrementalSerialization.h b/Engine/Plugins/Experimental/GameplayGraph/Source/GameplayGraph/Public/Graph/GraphIncrementalSerialization.h
new file mode 100644
--- /dev/null
+++ b/Engine/Plugins/Experimental/GameplayGraph/Source/GameplayGraph/Public/Graph/GraphIncrementalSerialization.h
@@ -0,0 +1,112 @@
+// Copyright Epic Games, Inc. All Rights Reserved.
+#pragma once
+
+#include "Graph/Graph.h"
+
+/*
+ * =========== README FIRST ===========
+ *
+ * This file provides the basics by which you can support incremental serialization for a UGraph for any possible
+ * serialized type. There are a few concepts at work here:
+ * 	- The "Graph" - this is the gameplay graph that you wish to serialize.
+ *  - The "Serializer" - this is the class that inherits from IGraphSerialization that serializes your data.
+ * 	  There is an additional caveat: this class must implement a "GetData" function that returns (see next point)...
+ *  - The "Serialized Graph Data" - this is the object that you wish to end up with after serialization.
+ * 
+ * With incremental serialization we introduce:
+ *  - "Delta Actions" - Any time the graph is changed, we add a "Delta Action" to a buffer. The next time Flush is called
+ *    on the incremental serialization, we'd expect the actions in the buffer to get applied to the cached data and make it
+ *    equivalent to the graph's current state.
+ *  - The "Delta Action Handler" - This takes in a list of "Delta Actions" and makes the appropriate changes to the "Serialized Graph Data".
+ *  - The "TGraphIncrementalSerialization" - this is the public interface to all of the above. Users of incremental serialization
+ *    need only construct a subclass of this to use. It should under-the-hood hook up to the graph events to properly detect changes.
+ */
+
+template<typename TDerived>
+struct TDerivedGraphDeltaAction
+{
+public:
+	template<typename TDeltaActionHandler>
+	void Accept(TDeltaActionHandler& Visitor, typename TDeltaActionHandler::TSerializableGraph& OutGraph) const
+	{
+		Visitor.Visit(*static_cast<const TDerived*>(this), OutGraph);
+	}
+};
+
+template<typename TInSerializableGraph>
+class TGraphDeltaActionHandler
+{
+	template<typename T>
+	struct AlwaysFalse
+	{
+	public:
+		static constexpr bool False = false;
+	};
+public:
+	using TSerializableGraph = TInSerializableGraph;
+
+	TGraphDeltaActionHandler()
+	{
+	}
+
+	virtual ~TGraphDeltaActionHandler() = default;
+
+	virtual void InitializeFromGraph(const TSerializableGraph& InGraph) {}
+	virtual void Flush(TSerializableGraph& OutGraph) {}
+
+	// Catch-all to remind users if their delta action handler doesn't actually support an action they created.
+	template<typename TAction>
+	void Visit(const TAction& Action, TSerializableGraph& OutGraph)
+	{
+		static_assert(AlwaysFalse<TAction>::False, "Unimplemented Visit function.");
+	}
+};
+
+template<typename TSerializer, typename TDeltaActionHandler>
+class TGraphIncrementalSerialization
+{
+public:
+	virtual ~TGraphIncrementalSerialization() {}
+
+	explicit TGraphIncrementalSerialization(UGraph* InGraph)
+	: Graph(InGraph)
+	{
+		if (!ensure(InGraph))
+		{
+			return;
+		}
+		// Build the serialized graph data from the current graph state and only do incremental updates from here on out.
+		TSerializer Serializer;
+		Serializer << *InGraph;
+		SerializedGraphData = Serializer.GetData();
+		DeltaActionHandler.InitializeFromGraph(SerializedGraphData);
+	}
+
+	/**
+     * This is not a const function as derived classes as the DeltaActionHandler may need to do a final pass on buffered data for example.
+     */
+	const typename TDeltaActionHandler::TSerializableGraph& GetLatestData()
+	{
+		DeltaActionHandler.Flush(SerializedGraphData);
+		return SerializedGraphData;
+	}
+
+protected:
+	template<typename TDeltaAction, typename... TParams>
+	void AddDeltaAction(TParams... Params)
+	{
+		TDeltaAction DeltaAction{Params...};
+		DeltaAction.Accept(DeltaActionHandler, SerializedGraphData);
+	}
+
+	UGraph* GetGraph() const
+	{
+		return Graph.Get();
+	}
+
+private:
+	TWeakObjectPtr<UGraph> Graph;
+	typename TDeltaActionHandler::TSerializableGraph SerializedGraphData;
+	// Note that the assumption here is that TDeltaActionHandler is the most derived type already so we don't need to use a pointer here.
+	TDeltaActionHandler DeltaActionHandler;
+};
\ No newline at end of file
diff --git a/Engine/Plugins/Experimental/GameplayGraph/Source/GameplayGraph/Public/Graph/GraphSerialization.h b/Engine/Plugins/Experimental/GameplayGraph/Source/GameplayGraph/Public/Graph/GraphSerialization.h
--- a/Engine/Plugins/Experimental/GameplayGraph/Source/GameplayGraph/Public/Graph/GraphSerialization.h
+++ b/Engine/Plugins/Experimental/GameplayGraph/Source/GameplayGraph/Public/Graph/GraphSerialization.h
@@ -21,6 +21,7 @@ public:
 	virtual void WriteGraphVertex(const FGraphVertexHandle& VertexHandle, const UGraphVertex* Vertex) = 0;
 	virtual void WriteGraphEdge(const FGraphVertexHandle& VertexHandleA, const FGraphVertexHandle& VertexHandleB) = 0;
 	virtual void WriteGraphIsland(const FGraphIslandHandle& IslandHandle, const UGraphIsland* Island) = 0;
+	virtual void Reset() = 0;
 };
 
 class IGraphDeserialization
diff --git a/Engine/Plugins/Experimental/GameplayGraph/Source/GameplayGraph/Public/Graph/GraphVertex.h b/Engine/Plugins/Experimental/GameplayGraph/Source/GameplayGraph/Public/Graph/GraphVertex.h
--- a/Engine/Plugins/Experimental/GameplayGraph/Source/GameplayGraph/Public/Graph/GraphVertex.h
+++ b/Engine/Plugins/Experimental/GameplayGraph/Source/GameplayGraph/Public/Graph/GraphVertex.h
@@ -6,8 +6,8 @@
 #include "GraphVertex.generated.h"
 
 /** Event for when the node has been removed from the graph. */
-DECLARE_MULTICAST_DELEGATE(FOnGraphVertexRemoved);
-DECLARE_MULTICAST_DELEGATE_OneParam(FOnGraphVertexParentIslandSet, const FGraphIslandHandle&);
+DECLARE_MULTICAST_DELEGATE_OneParam(FOnGraphVertexRemoved, const FGraphVertexHandle&);
+DECLARE_MULTICAST_DELEGATE_TwoParams(FOnGraphVertexParentIslandSet, const FGraphVertexHandle&, const FGraphIslandHandle&);
 
 UCLASS()
 class GAMEPLAYGRAPH_API UGraphVertex : public UGraphElement
diff --git a/Engine/Plugins/Experimental/GameplayGraph/Tests/Private/GraphSerializationTests.cpp b/Engine/Plugins/Experimental/GameplayGraph/Tests/Private/GraphSerializationTests.cpp
--- a/Engine/Plugins/Experimental/GameplayGraph/Tests/Private/GraphSerializationTests.cpp
+++ b/Engine/Plugins/Experimental/GameplayGraph/Tests/Private/GraphSerializationTests.cpp
@@ -377,4 +377,96 @@ TEST_CASE_METHOD(FTestGraphBuilder, "Graph::Serialization::Read::Errors::Island:
 	CHECK(LoadedIsland->Num() == 2);
 	CHECK(LoadedIsland->GetVertices().Contains(SerializedGraph.Vertices[0]) == true);
 	CHECK(LoadedIsland->GetVertices().Contains(SerializedGraph.Vertices[1]) == true);
+}
+
+TEST_CASE_METHOD(FTestGraphBuilder, "Graph::Serialization::Incremental::Initialization", "[graph][serialization][incremental]")
+{
+	PopulateVertices(6, true);
+	BuildLinearEdges(3);
+
+	CHECK(Graph->NumVertices() == 6);
+	CHECK(Graph->NumIslands() == 2);
+
+	FDefaultGraphSerialization Serializer;
+	Serializer << *Graph;
+
+	FDefaultGraphIncrementalSerialization IncrementalSerializer(Graph);
+	CHECK(Serializer.GetData() == IncrementalSerializer.GetLatestData());
+}
+
+TEST_CASE_METHOD(FTestGraphBuilder, "Graph::Serialization::Incremental::GraphBuild", "[graph][serialization][incremental]")
+{
+	SECTION("Linear")
+	{
+		FDefaultGraphIncrementalSerialization IncrementalSerializer(Graph);
+		PopulateVertices(6, true);
+		BuildLinearEdges(3);
+
+		CHECK(Graph->NumVertices() == 6);
+		CHECK(Graph->NumIslands() == 2);
+
+		FDefaultGraphSerialization Serializer;
+		Serializer << *Graph;
+		CHECK(Serializer.GetData() == IncrementalSerializer.GetLatestData());
+	}
+
+	SECTION("Fully Connected")
+	{
+		FDefaultGraphIncrementalSerialization IncrementalSerializer(Graph);
+		PopulateVertices(10, true);
+		BuildFullyConnectedEdges(5);
+
+		CHECK(Graph->NumVertices() == 10);
+		CHECK(Graph->NumIslands() == 2);
+
+		FDefaultGraphSerialization Serializer;
+		Serializer << *Graph;
+		CHECK(Serializer.GetData() == IncrementalSerializer.GetLatestData());
+	}
+}
+
+TEST_CASE_METHOD(FTestGraphBuilder, "Graph::Serialization::Incremental::Removal", "[graph][serialization][incremental]")
+{
+	PopulateVertices(10, true);
+	BuildFullyConnectedEdges(5);
+	FDefaultGraphIncrementalSerialization IncrementalSerializer(Graph);
+
+	SECTION("Remove Vertices")
+	{
+		Graph->RemoveVertex(VertexHandles[0]);
+		Graph->RemoveVertex(VertexHandles[1]);
+		Graph->RemoveVertex(VertexHandles[6]);
+
+		{
+			FDefaultGraphSerialization Serializer;
+			Serializer << *Graph;
+			CHECK(Serializer.GetData() == IncrementalSerializer.GetLatestData());
+		}
+
+		TArray<FGraphVertexHandle> Bulk{ VertexHandles[2], VertexHandles[7], VertexHandles[8] };
+		Graph->RemoveBulkVertices(Bulk);
+
+		{
+			FDefaultGraphSerialization Serializer;
+			Serializer << *Graph;
+			CHECK(Serializer.GetData() == IncrementalSerializer.GetLatestData());
+		}
+	}
+
+	SECTION("Remove Island")
+	{
+		Graph->RemoveIsland(IslandHandles[0]);
+		{
+			FDefaultGraphSerialization Serializer;
+			Serializer << *Graph;
+			CHECK(Serializer.GetData() == IncrementalSerializer.GetLatestData());
+		}
+
+		Graph->RemoveIsland(IslandHandles[0]);
+		{
+			FDefaultGraphSerialization Serializer;
+			Serializer << *Graph;
+			CHECK(Serializer.GetData() == IncrementalSerializer.GetLatestData());
+		}
+	}
 }
\ No newline at end of file
diff --git a/Engine/Plugins/Experimental/GameplayGraph/Tests/Private/GraphVertexTests.cpp b/Engine/Plugins/Experimental/GameplayGraph/Tests/Private/GraphVertexTests.cpp
--- a/Engine/Plugins/Experimental/GameplayGraph/Tests/Private/GraphVertexTests.cpp
+++ b/Engine/Plugins/Experimental/GameplayGraph/Tests/Private/GraphVertexTests.cpp
@@ -248,8 +248,9 @@ TEST_CASE_METHOD(FTestGraphBuilder, "Graph::Vertex::Event::Remove", "[graph][ver
 		UGraphVertex* Vertex = VertexHandle.GetVertex();
 		REQUIRE(Vertex != nullptr);
 		Vertex->OnVertexRemoved.AddLambda(
-			[Index, &VertexRemoved]()
+			[Vertex, Index, &VertexRemoved](const FGraphVertexHandle& Handle)
 			{
+				CHECK(Vertex->Handle() == Handle);
 				VertexRemoved[Index] = true;
 			}
 		);
@@ -291,8 +292,9 @@ TEST_CASE_METHOD(FTestGraphBuilder, "Graph::Vertex::Event::Parent Island Set", "
 		UGraphVertex* Vertex = VertexHandle.GetVertex();
 		REQUIRE(Vertex != nullptr);
 		Vertex->OnParentIslandSet.AddLambda(
-			[Index, &VertexIslandSet](const FGraphIslandHandle& IslandHandle)
+			[Vertex, Index, &VertexIslandSet](const FGraphVertexHandle& VertexHandle, const FGraphIslandHandle& IslandHandle)
 			{
+				CHECK(Vertex->Handle() == VertexHandle);
 				VertexIslandSet[Index] = IslandHandle;
 			}
 		);
