Add explicit default constructor to FGraphProperties

diff --git a/Engine/Plugins/Experimental/GameplayGraph/Source/GameplayGraph/Public/Graph/Graph.h b/Engine/Plugins/Experimental/GameplayGraph/Source/GameplayGraph/Public/Graph/Graph.h
--- a/Engine/Plugins/Experimental/GameplayGraph/Source/GameplayGraph/Public/Graph/Graph.h
+++ b/Engine/Plugins/Experimental/GameplayGraph/Source/GameplayGraph/Public/Graph/Graph.h
@@ -21,6 +21,8 @@ struct FGraphProperties
 {
 	GENERATED_BODY()
 
+	FGraphProperties() = default;
+ 
 	UPROPERTY(SaveGame)
 	bool bGenerateIslands = true;
 	
