Graph ChangeVertexHandle.
Used to migrate actor GUIDs.
fyi gautam.muralidhar

diff --git a/Engine/Plugins/Experimental/GameplayGraph/Source/GameplayGraph/Private/Graph/Graph.cpp b/Engine/Plugins/Experimental/GameplayGraph/Source/GameplayGraph/Private/Graph/Graph.cpp
--- a/Engine/Plugins/Experimental/GameplayGraph/Source/GameplayGraph/Private/Graph/Graph.cpp
+++ b/Engine/Plugins/Experimental/GameplayGraph/Source/GameplayGraph/Private/Graph/Graph.cpp
@@ -64,6 +64,38 @@ FGraphVertexHandle UGraph::CreateVertex(FGraphUniqueIndex InUniqueIndex)
 	return Vertex->Handle();
 }
 
+void UGraph::ChangeVertexHandle(const FGraphVertexHandle& OldVertexHandle, const FGraphVertexHandle& NewVertexHandle)
+{
+	UGraphVertex* Vertex = GetSafeVertexFromHandle(OldVertexHandle);
+	if (!ensure(Vertex))
+	{
+		return;
+	}
+
+	// Update vertex
+	Vertex->SetUniqueIndex(NewVertexHandle.GetUniqueIndex());
+
+	Vertices.Remove(OldVertexHandle);
+	Vertices.Add(NewVertexHandle, Vertex);
+
+	// Update edges. We use that it is bidirectional to only modify the edges of the adjacent vertices. 
+	Vertex->ForEachAdjacentVertex(
+		[this, &OldVertexHandle, &NewVertexHandle](const FGraphVertexHandle& EdgeVertexHandle)
+		{
+			UGraphVertex* EdgeNode = GetSafeVertexFromHandle(EdgeVertexHandle);
+			if (ensure(EdgeNode))
+			{
+				EdgeNode->ChangeEdgeVertexHandle(OldVertexHandle, NewVertexHandle);
+			}
+		});
+
+	// Update island
+	if (UGraphIsland* Island = Vertex->GetParentIsland().GetIsland())
+	{
+		Island->ChangeVertexHandle(OldVertexHandle, NewVertexHandle);
+	}
+}
+
 void UGraph::RegisterVertex(TObjectPtr<UGraphVertex> Vertex)
 {
 	TRACE_CPUPROFILER_EVENT_SCOPE(UGraph::RegisterVertex);
diff --git a/Engine/Plugins/Experimental/GameplayGraph/Source/GameplayGraph/Private/Graph/GraphIsland.cpp b/Engine/Plugins/Experimental/GameplayGraph/Source/GameplayGraph/Private/Graph/GraphIsland.cpp
--- a/Engine/Plugins/Experimental/GameplayGraph/Source/GameplayGraph/Private/Graph/GraphIsland.cpp
+++ b/Engine/Plugins/Experimental/GameplayGraph/Source/GameplayGraph/Private/Graph/GraphIsland.cpp
@@ -91,6 +91,14 @@ void UGraphIsland::RemoveVertex(const FGraphVertexHandle& Node)
 	}
 }
 
+void UGraphIsland::ChangeVertexHandle(const FGraphVertexHandle& OldVertexHandle, const FGraphVertexHandle& NewVertexHandle)
+{
+	if (ensure(Vertices.Remove(OldVertexHandle)))
+	{
+		Vertices.Add(NewVertexHandle);
+	}
+}
+
 void UGraphIsland::HandleOnVertexAdded(const FGraphVertexHandle& Handle)
 {
 	OnVertexAdded.Broadcast(this->Handle(), Handle);
diff --git a/Engine/Plugins/Experimental/GameplayGraph/Source/GameplayGraph/Private/Graph/GraphVertex.cpp b/Engine/Plugins/Experimental/GameplayGraph/Source/GameplayGraph/Private/Graph/GraphVertex.cpp
--- a/Engine/Plugins/Experimental/GameplayGraph/Source/GameplayGraph/Private/Graph/GraphVertex.cpp
+++ b/Engine/Plugins/Experimental/GameplayGraph/Source/GameplayGraph/Private/Graph/GraphVertex.cpp
@@ -30,6 +30,14 @@ void UGraphVertex::RemoveEdge(const FGraphVertexHandle& AdjacentVertexHandle)
 	Edges.Remove(AdjacentVertexHandle);
 }
 
+void UGraphVertex::ChangeEdgeVertexHandle(const FGraphVertexHandle& OldVertexHandle, const FGraphVertexHandle& NewVertexHandle)
+{
+	if (ensure(Edges.Remove(OldVertexHandle)))
+	{
+		Edges.Add(NewVertexHandle);
+	}
+}
+
 void UGraphVertex::HandleOnVertexRemoved()
 {
 	OnVertexRemoved.Broadcast(Handle());
diff --git a/Engine/Plugins/Experimental/GameplayGraph/Source/GameplayGraph/Public/Graph/Graph.h b/Engine/Plugins/Experimental/GameplayGraph/Source/GameplayGraph/Public/Graph/Graph.h
--- a/Engine/Plugins/Experimental/GameplayGraph/Source/GameplayGraph/Public/Graph/Graph.h
+++ b/Engine/Plugins/Experimental/GameplayGraph/Source/GameplayGraph/Public/Graph/Graph.h
@@ -96,6 +96,9 @@ public:
 	/** Create a node with the specified subclass, adds it to the graph, and returns a handle to it. */
 	FGraphVertexHandle CreateVertex(FGraphUniqueIndex InUniqueIndex = FGraphUniqueIndex::CreateUniqueIndex(false));
 
+	/** Changes a node handle, updating all the references. The vertex must exist. */
+	void ChangeVertexHandle(const FGraphVertexHandle& OldVertexHandle, const FGraphVertexHandle& NewVertexHandle);
+
 	/** Creates edges in bulk. This is more efficient than calling CreateEdge multiple times since we will only try to assign a node to an island once. */
 	void CreateBulkEdges(TArray<FEdgeSpecifier>&& NodesToConnect);
 
diff --git a/Engine/Plugins/Experimental/GameplayGraph/Source/GameplayGraph/Public/Graph/GraphIsland.h b/Engine/Plugins/Experimental/GameplayGraph/Source/GameplayGraph/Public/Graph/GraphIsland.h
--- a/Engine/Plugins/Experimental/GameplayGraph/Source/GameplayGraph/Public/Graph/GraphIsland.h
+++ b/Engine/Plugins/Experimental/GameplayGraph/Source/GameplayGraph/Public/Graph/GraphIsland.h
@@ -85,6 +85,9 @@ public:
 	/** Removes a node from the island. */
 	void RemoveVertex(const FGraphVertexHandle& Node);
 
+	/** Changes a node handle. The vertex must exist. */
+	virtual void ChangeVertexHandle(const FGraphVertexHandle& OldVertexHandle, const FGraphVertexHandle& NewVertexHandle);
+
 	FOnGraphIslandVertexAdded OnVertexAdded;
 	FOnGraphIslandVertexRemoved OnVertexRemoved;
 	FOnGraphIslandDestroyed OnDestroyed;
diff --git a/Engine/Plugins/Experimental/GameplayGraph/Source/GameplayGraph/Public/Graph/GraphVertex.h b/Engine/Plugins/Experimental/GameplayGraph/Source/GameplayGraph/Public/Graph/GraphVertex.h
--- a/Engine/Plugins/Experimental/GameplayGraph/Source/GameplayGraph/Public/Graph/GraphVertex.h
+++ b/Engine/Plugins/Experimental/GameplayGraph/Source/GameplayGraph/Public/Graph/GraphVertex.h
@@ -37,6 +37,9 @@ public:
 	
 	const TSet<FGraphVertexHandle>& GetEdges() const { return Edges; }
 
+	/** Changes an edge vertex handle. The vertex must exist. */
+	void ChangeEdgeVertexHandle(const FGraphVertexHandle& OldVertexHandle, const FGraphVertexHandle& NewVertexHandle);
+
 	FOnGraphVertexRemoved OnVertexRemoved;
 	FOnGraphVertexParentIslandSet OnParentIslandSet;
 
