Re-expose the RemoveEdge function in UGraph, which was removed in 31844651 in order to protect against calling remove edge without updating islands, and as a part of an optimization to no longer store a whole list of edges. This exposed version does not allow callers to call with bHandleIslands = false.

diff --git a/Engine/Plugins/Experimental/GameplayGraph/Source/GameplayGraph/Private/Graph/Graph.cpp b/Engine/Plugins/Experimental/GameplayGraph/Source/GameplayGraph/Private/Graph/Graph.cpp
--- a/Engine/Plugins/Experimental/GameplayGraph/Source/GameplayGraph/Private/Graph/Graph.cpp
+++ b/Engine/Plugins/Experimental/GameplayGraph/Source/GameplayGraph/Private/Graph/Graph.cpp
@@ -549,6 +549,11 @@ void UGraph::RemoveBulkVertices(const TArray<FGraphVertexHandle>& InHandles)
 	}
 }
 
+void UGraph::RemoveEdge(const FGraphVertexHandle& VertexHandleA, const FGraphVertexHandle& VertexHandleB)
+{
+	RemoveEdgeInternal(VertexHandleA, VertexHandleB, true);
+}
+
 void UGraph::RemoveEdgeInternal(const FGraphVertexHandle& VertexHandleA, const FGraphVertexHandle& VertexHandleB, bool bHandleIslands)
 {
 	TRACE_CPUPROFILER_EVENT_SCOPE(UGraph::RemoveEdgeInternal);
diff --git a/Engine/Plugins/Experimental/GameplayGraph/Source/GameplayGraph/Public/Graph/Graph.h b/Engine/Plugins/Experimental/GameplayGraph/Source/GameplayGraph/Public/Graph/Graph.h
--- a/Engine/Plugins/Experimental/GameplayGraph/Source/GameplayGraph/Public/Graph/Graph.h
+++ b/Engine/Plugins/Experimental/GameplayGraph/Source/GameplayGraph/Public/Graph/Graph.h
@@ -107,6 +107,9 @@ public:
 
 	void RemoveBulkVertices(const TArray<FGraphVertexHandle>& InHandles);
 
+	/** Removes an edge between two nodes */
+	void RemoveEdge(const FGraphVertexHandle& VertexHandleA, const FGraphVertexHandle& VertexHandleB);
+
 	/** This should be called immediately after a node and any relevant edges have been added to the graph. */
 	virtual void FinalizeVertex(const FGraphVertexHandle& InHandle);
 
