Fix incorrect island reference in UGraph::CreateEdge

diff --git a/Engine/Plugins/Experimental/GameplayGraph/Source/GameplayGraph/Private/Graph/Graph.cpp b/Engine/Plugins/Experimental/GameplayGraph/Source/GameplayGraph/Private/Graph/Graph.cpp
--- a/Engine/Plugins/Experimental/GameplayGraph/Source/GameplayGraph/Private/Graph/Graph.cpp
+++ b/Engine/Plugins/Experimental/GameplayGraph/Source/GameplayGraph/Private/Graph/Graph.cpp
@@ -182,7 +182,7 @@ bool UGraph::CreateEdge(FGraphVertexHandle Node1, FGraphVertexHandle Node2, bool
 			{
 				Island1->AddVertex(Node2);
 			}
-			else if (UGraphIsland* Island2 = IslandHandle1.GetIsland())
+			else if (UGraphIsland* Island2 = IslandHandle2.GetIsland())
 			{
 				Island2->AddVertex(Node1);
 			}
