UGraph serialization optimization

5? speedup in UGraph serialization.

- Backwards & forwards compatible with existing serialized data.
- Clarify that edges are serialized in only one direction; the reverse is implicit (and created during deserialization as before). UGraph is undirected; edges are always bidirectional.
- Improve estimation of number of edges in graph for memory reservation.
- Inline cheap FGraphHandle comparisons.
- Move expensive ensure into a checkSlow().
- Remove redundant cache of seen edges, replaced by a check of the order of vertices.

#tests Unit tests 100% passing. Tested with real world game data.
#rnx

diff --git a/Engine/Plugins/Experimental/GameplayGraph/Source/GameplayGraph/Private/Graph/Graph.cpp b/Engine/Plugins/Experimental/GameplayGraph/Source/GameplayGraph/Private/Graph/Graph.cpp
--- a/Engine/Plugins/Experimental/GameplayGraph/Source/GameplayGraph/Private/Graph/Graph.cpp
+++ b/Engine/Plugins/Experimental/GameplayGraph/Source/GameplayGraph/Private/Graph/Graph.cpp
@@ -6,6 +6,7 @@
 #include "Graph/GraphSerialization.h"
 #include "GenericPlatform/GenericPlatformMath.h"
 #include "HAL/IConsoleManager.h"
+#include "Templates/UnrealTemplate.h"
 
 #include UE_INLINE_GENERATED_CPP_BY_NAME(Graph)
 
@@ -13,15 +14,26 @@ namespace Graph
 {
 	static bool bFixupGraphOnLoad = true;
 	static FAutoConsoleVariableRef FixupGraphOnLoadCVar(TEXT("GameplayGraph.FixupGraphOnLoad"), bFixupGraphOnLoad, TEXT("Merge and split islands when loaded from serialization."));
-}
 
-FEdgeSpecifier::FEdgeSpecifier(const FGraphVertexHandle& InVertexHandle1, const FGraphVertexHandle& InVertexHandle2)
-: VertexHandle1(InVertexHandle1)
-, VertexHandle2(InVertexHandle2)
-{
-	if (VertexHandle2 < VertexHandle1)
+	// Return a rough estimate of the total number of edges in the graph. This can be used to reserve memory for each edge.
+	// The return type is int32 to match IGraphSerialization::Initialize().
+	static int32 EstimateNumEdges(const UGraph& Graph)
 	{
-		std::swap(VertexHandle1, VertexHandle2);
+		const TMap<FGraphVertexHandle, TObjectPtr<UGraphVertex>>& Vertices = Graph.GetVertices();
+		if (Vertices.IsEmpty())
+		{
+			return 0;
+		}
+
+		const TObjectPtr<UGraphVertex>& FirstVertex = Vertices.CreateConstIterator()->Value;
+		if (UNLIKELY(!FirstVertex))
+		{
+			return 0;
+		}
+
+		// Clamp estimate to reduce the chance of a specially crafted graph or atypical sample vertex leading to excessive overallocation.
+		uint64 Estimate = static_cast<uint64>(Vertices.Num()) * FMath::Clamp(FirstVertex->NumEdges(), 2, 10);
+		return IntCastChecked<int32>(FMath::Min<uint64>(Estimate, std::numeric_limits<int32>::max()));
 	}
 }
 
@@ -758,18 +770,19 @@ void operator<<(IGraphSerialization& Output, const UGraph& Graph)
 {
 	TRACE_CPUPROFILER_EVENT_SCOPE_STR("Gameplay Graph Serialization");
 
-	// We don't actually have a good estimate of the number of edges anymore - so just guesstimating number of vertices.
-	Output.Initialize(Graph.NumVertices(), Graph.NumVertices(), Graph.NumIslands());
+	Output.Initialize(Graph.NumVertices(), Graph::EstimateNumEdges(Graph), Graph.NumIslands());
 	Output.WriteGraphProperties(Graph.GetProperties());
 
 	{
 		TRACE_CPUPROFILER_EVENT_SCOPE_STR("Serialize Vertices and Edges");
-		TSet<FEdgeSpecifier> SeenEdges;
-		SeenEdges.Reserve(Graph.NumVertices());
 
 		for (const TPair<FGraphVertexHandle, TObjectPtr<UGraphVertex>>& Kvp : Graph.GetVertices())
 		{
-			if (!ensure(Kvp.Key.IsComplete() && Kvp.Value))
+			// Vertices are directly from the graph, keys should always be complete. Checking IsComplete() looks up the key in the map, which is
+			// expensive when doing it for every vertex in the graph.
+			checkSlow(Kvp.Key.IsComplete());
+
+			if (!ensure(Kvp.Value))
 			{
 				continue;
 			}
@@ -777,14 +790,17 @@ void operator<<(IGraphSerialization& Output, const UGraph& Graph)
 			Output.WriteGraphVertex(Kvp.Key, Kvp.Value);
 
 			Kvp.Value->ForEachAdjacentVertex(
-				[&SeenEdges, &Kvp, &Output](const FGraphVertexHandle& AdjacentVertexHandle)
+				[&Kvp, &Output](const FGraphVertexHandle& AdjacentVertexHandle)
 				{
-					FEdgeSpecifier Edge{ Kvp.Key, AdjacentVertexHandle };
-					if (!SeenEdges.Contains(Edge))
+					// Only serialize the vertex in one direction, the reverse is implicit and re-created during deserialization.
+					// Note that this permits edges from one vertex back to itself (potentially several such edges), which shall be serialized.
+					if (AdjacentVertexHandle < Kvp.Key)
 					{
-						SeenEdges.Add(Edge);
-						Output.WriteGraphEdge(Edge.GetVertexHandle1(), Edge.GetVertexHandle2());
+						// This edge is already serialized in the reverse direction as {AdjacentVertexHandle, Kvp.Key}
+						return;
 					}
+
+					Output.WriteGraphEdge(Kvp.Key, AdjacentVertexHandle);
 				}
 			);
 		}
diff --git a/Engine/Plugins/Experimental/GameplayGraph/Source/GameplayGraph/Private/Graph/GraphHandle.cpp b/Engine/Plugins/Experimental/GameplayGraph/Source/GameplayGraph/Private/Graph/GraphHandle.cpp
--- a/Engine/Plugins/Experimental/GameplayGraph/Source/GameplayGraph/Private/Graph/GraphHandle.cpp
+++ b/Engine/Plugins/Experimental/GameplayGraph/Source/GameplayGraph/Private/Graph/GraphHandle.cpp
@@ -21,42 +21,12 @@ FGraphHandle::FGraphHandle(FGraphUniqueIndex InUniqueIndex, UGraph* InGraph)
 {
 }
 
-bool FGraphHandle::IsValid() const
-{
-	return UniqueIndex.IsValid();
-}
-
-bool FGraphHandle::IsComplete() const
-{
-	return IsValid() && HasElement();
-}
-
-bool FGraphHandle::operator==(const FGraphHandle& Other) const
-{
-	return UniqueIndex == Other.UniqueIndex;
-}
-
-bool FGraphHandle::operator!=(const FGraphHandle& Other) const
-{
-	return UniqueIndex != Other.UniqueIndex;
-}
-
-bool FGraphHandle::operator<(const FGraphHandle& Other) const
-{
-	return UniqueIndex < Other.UniqueIndex;
-}
-
 void FGraphHandle::Clear()
 {
 	UniqueIndex = FGraphUniqueIndex();
 	WeakGraph = nullptr;
 }
 
-uint32 GetTypeHash(const FGraphHandle& Handle)
-{
-	return GetTypeHash(Handle.UniqueIndex);
-}
-
 UGraphVertex* FGraphVertexHandle::GetVertex() const
 {
 	if (UGraph* Graph = GetGraph())
diff --git a/Engine/Plugins/Experimental/GameplayGraph/Source/GameplayGraph/Public/Graph/Graph.h b/Engine/Plugins/Experimental/GameplayGraph/Source/GameplayGraph/Public/Graph/Graph.h
--- a/Engine/Plugins/Experimental/GameplayGraph/Source/GameplayGraph/Public/Graph/Graph.h
+++ b/Engine/Plugins/Experimental/GameplayGraph/Source/GameplayGraph/Public/Graph/Graph.h
@@ -36,9 +36,15 @@ struct FGraphProperties
 DECLARE_MULTICAST_DELEGATE_OneParam(FOnGraphVertexCreated, const FGraphVertexHandle&);
 DECLARE_MULTICAST_DELEGATE_OneParam(FOnGraphIslandCreated, const FGraphIslandHandle&);
 
+// Undirected edge descriptor.
+// Always VertexHandle1 <= VertexHandle2.
 struct FEdgeSpecifier
 {
 	UE_API FEdgeSpecifier(const FGraphVertexHandle& InVertexHandle1, const FGraphVertexHandle& InVertexHandle2);
+	FEdgeSpecifier(const FEdgeSpecifier&) = default;
+	FEdgeSpecifier(FEdgeSpecifier&&) = default;
+	FEdgeSpecifier& operator=(const FEdgeSpecifier&) = default;
+	FEdgeSpecifier& operator=(FEdgeSpecifier&&) = default;
 
 	// Comparison operators
 	friend bool operator==(const FEdgeSpecifier& Lhs, const FEdgeSpecifier& Rhs) = default;
@@ -187,3 +193,13 @@ private:
 };
 
 #undef UE_API
+
+inline FEdgeSpecifier::FEdgeSpecifier(const FGraphVertexHandle& InVertexHandle1, const FGraphVertexHandle& InVertexHandle2)
+	: VertexHandle1(InVertexHandle1)
+	, VertexHandle2(InVertexHandle2)
+{
+	if (VertexHandle2 < VertexHandle1)
+	{
+		std::swap(VertexHandle1, VertexHandle2);
+	}
+}
diff --git a/Engine/Plugins/Experimental/GameplayGraph/Source/GameplayGraph/Public/Graph/GraphHandle.h b/Engine/Plugins/Experimental/GameplayGraph/Source/GameplayGraph/Public/Graph/GraphHandle.h
--- a/Engine/Plugins/Experimental/GameplayGraph/Source/GameplayGraph/Public/Graph/GraphHandle.h
+++ b/Engine/Plugins/Experimental/GameplayGraph/Source/GameplayGraph/Public/Graph/GraphHandle.h
@@ -165,3 +165,33 @@ struct FGraphIslandHandle : public FGraphHandle
 };
 
 #undef UE_API
+
+inline bool FGraphHandle::IsValid() const
+{
+	return UniqueIndex.IsValid();
+}
+
+inline bool FGraphHandle::IsComplete() const
+{
+	return IsValid() && HasElement();
+}
+
+inline bool FGraphHandle::operator==(const FGraphHandle& Other) const
+{
+	return UniqueIndex == Other.UniqueIndex;
+}
+
+inline bool FGraphHandle::operator!=(const FGraphHandle& Other) const
+{
+	return UniqueIndex != Other.UniqueIndex;
+}
+
+inline bool FGraphHandle::operator<(const FGraphHandle& Other) const
+{
+	return UniqueIndex < Other.UniqueIndex;
+}
+
+inline uint32 GetTypeHash(const FGraphHandle& Handle)
+{
+	return GetTypeHash(Handle.UniqueIndex);
+}
